Branch snapshot for CL 3327560

diff --git a/Engine/Source/Editor/MeshEditor/MeshEditorMode.cpp b/Engine/Source/Editor/MeshEditor/MeshEditorMode.cpp
new file mode 100644
--- /dev/null
+++ b/Engine/Source/Editor/MeshEditor/MeshEditorMode.cpp
@@ -0,0 +1,6431 @@
+// Copyright 1998-2016 Epic Games, Inc. All Rights Reserved.
+
+#include "MeshEditorMode.h"
+#include "MeshEditorCommands.h"
+#include "MeshEditorModeToolkit.h"
+#include "EditableMesh.h"
+#include "EditableMeshFactory.h"
+#include "MeshElement.h"
+#include "DynamicMeshBuilder.h"
+#include "ScopedTransaction.h"
+#include "SnappingUtils.h"
+#include "GeomTools.h"
+#include "Toolkits/ToolkitManager.h"
+#include "SEditorViewport.h"
+#include "ContentBrowserModule.h"
+#include "EditorViewportClient.h"
+#include "EditorWorldExtension.h"
+#include "MeshElementViewportTransformable.h"
+#include "MeshElementTransformer.h"
+#include "ViewportInteractor.h"
+#include "ActorViewportTransformable.h"
+#include "ViewportWorldInteraction.h"
+#include "IViewportInteractionModule.h"
+#include "IContentBrowserSingleton.h"
+#include "VIBaseTransformGizmo.h"	// For EGizmoHandleTypes
+#include "Engine/Selection.h"
+#include "EngineUtils.h"
+#include "DrawDebugHelpers.h"	// For DrawDebugSphere
+#include "MeshEditorSettings.h"
+#include "LevelEditor.h"
+#include "VREditorMode.h" // @TODO: remove once radial menu code is in a non-VREditor module
+#include "MultiBoxBuilder.h"
+#include "UICommandList.h"
+#include "VREditorWorldInteraction.h"
+#include "MeshEditorAssetContainer.h"
+#include "Kismet/GameplayStatics.h"
+
+#define LOCTEXT_NAMESPACE "MeshEditorMode"
+
+// @todo mesheditor urgent: Texture coords are getting trashed all the time (only for meshes with 1 UV set though?)
+
+namespace MeshEd
+{
+	static FAutoConsoleVariable HoverFadeDuration( TEXT( "MeshEd.HoverFadeDuration" ), 0.3f, TEXT( "How many seconds over which we should fade out hovered mesh elements." ) );
+	static FAutoConsoleVariable SelectionAnimationDuration( TEXT( "MeshEd.SelectionAnimationDuration" ), 0.2f, TEXT( "How long the animation should last when selecting a mesh element." ) );
+	static FAutoConsoleVariable MinDeltaForInertialMovement( TEXT( "MeshEd.MinDeltaForInertialMovement" ), 0.01f, TEXT( "Minimum velocity in cm/frame for inertial movement to kick in when releasing a drag" ) );
+	static FAutoConsoleVariable ShowDebugStats( TEXT( "MeshEd.ShowDebugStats" ), 0, TEXT( "Enables debug overlay text for the currently selected mesh" ) );
+	static FAutoConsoleVariable EnableSelectByPainting( TEXT( "MeshEd.EnableSelectByPainting" ), 0, TEXT( "Enables selection by clicking and dragging over elements" ) );
+
+	static FAutoConsoleVariable OverlayDepthOffset( TEXT( "MeshEd.OverlayDepthOffset" ), 2.0f, TEXT( "How far to offset overlay wires/polygons on top of meshes when hovered or selected" ) );
+	static FAutoConsoleVariable OverlayVertexSize( TEXT( "MeshEd.OverlayVertexSize" ), 4.0f, TEXT( "How large a vertex is on a mesh overlay" ) );
+	static FAutoConsoleVariable OverlayLineThickness( TEXT( "MeshEd.OverlayLineThickness" ), 0.9f, TEXT( "How thick overlay lines should be on top of meshes when hovered or selected" ) );
+	static FAutoConsoleVariable OverlayDistanceScaleFactor( TEXT( "MeshEd.OverlayDistanceScaleFactor" ), 0.002f, TEXT( "How much to scale overlay wires automatically based on distance to the viewer" ) );
+	static FAutoConsoleVariable OverlayDistanceBias( TEXT( "MeshEd.OverlayDistanceBias" ), 0.05f, TEXT( "How much to bias distance scale by, regardless of distance to the viewer" ) );
+	static FAutoConsoleVariable SelectedSizeBias( TEXT( "MeshEd.SelectedSizeBias" ), 0.1f, TEXT( "Selected mesh element size bias" ) );
+	static FAutoConsoleVariable SelectedAnimationExtraSizeBias( TEXT( "MeshEd.SelectedAnimationExtraSizeBias" ), 2.5f, TEXT( "Extra hovered mesh element size bias when animating" ) );
+	static FAutoConsoleVariable HoveredSizeBias( TEXT( "MeshEd.HoveredSizeBias" ), 0.1f, TEXT( "Selected mesh element size bias" ) );
+	static FAutoConsoleVariable HoveredAnimationExtraSizeBias( TEXT( "MeshEd.HoveredAnimationExtraSizeBias" ), 0.5f, TEXT( "Extra hovered mesh element size bias when animating" ) );
+	
+	static FAutoConsoleVariable LaserFuzzySelectionDistance( TEXT( "MeshEd.LaserFuzzySelectionDistance" ), 4.0f, TEXT( "Distance in world space to allow selection of mesh elements using laser, even when not directly over them" ) );
+	static FAutoConsoleVariable GrabberSphereFuzzySelectionDistance( TEXT( "MeshEd.GrabberSphereFuzzySelectionDistance" ), 2.0f, TEXT( "Distance in world space to allow selection of mesh elements using grabber sphere, even when not directly over them" ) );
+	static FAutoConsoleVariable SFXMultiplier(TEXT("MeshEd.SFXMultiplier"), 2.0f, TEXT("Default Sound Effect Volume Multiplier"));
+}
+
+
+TUniquePtr<FChange> FMeshEditorMode::FSelectOrDeselectMeshElementsChange::Execute( UObject* Object )
+{
+	// @todo mesheditor urgent: What if mode is EXITED and user presses Ctrl+Z!  Force pending kill and skip?  Currently, the FEdMode object is not destroyed.  So it sort of just works.
+	//     --> Should selection persist after exiting the mode?  It's weird that undo/redo won't show you changes...  Not sure though.  This is the sort of thing needed if we want whole-application-undo/redo though.  Even a mode switch belongs in the undo buffer.
+	UMeshEditorModeProxyObject* LocalMeshEditorModeProxyObject = CastChecked<UMeshEditorModeProxyObject>( Object );
+	FMeshEditorMode& MeshEditorMode = *LocalMeshEditorModeProxyObject->OwningMeshEditorMode;
+
+	// Get the current element selection mode
+	const EEditableMeshElementType CurrentElementSelectionMode = MeshEditorMode.MeshElementSelectionMode;
+
+	// Back up the current selection so we can restore it on undo
+	FCompoundChangeInput CompoundRevertInput;
+
+	FSelectOrDeselectMeshElementsChangeInput RevertInput;
+	RevertInput.MeshElementsToSelect = Input.MeshElementsToDeselect;
+	RevertInput.MeshElementsToDeselect = Input.MeshElementsToSelect;
+	CompoundRevertInput.Subchanges.Add( MakeUnique<FSelectOrDeselectMeshElementsChange>( MoveTemp( RevertInput ) ) );
+
+	const double CurrentRealTime = FSlateApplication::Get().GetCurrentTime();
+
+	// Selection changed.  This is a good time to reset the hover animation time value, to avoid problems with floating point precision
+	// when it gets too large.
+	MeshEditorMode.HoverFeedbackTimeValue = 0.0;
+
+	if( MeshEditorMode.IsActive() )
+	{
+		for( FMeshElement& MeshElementToDeselect : Input.MeshElementsToDeselect )
+		{
+			const int32 RemoveAtIndex = MeshEditorMode.GetSelectedMeshElementIndex( MeshElementToDeselect );
+			if( RemoveAtIndex != INDEX_NONE )
+			{
+				MeshEditorMode.SelectedMeshElements.RemoveAtSwap( RemoveAtIndex );
+			}
+		}
+
+		if( Input.MeshElementsToSelect.Num() > 0 )
+		{
+			// Make sure they're all the same type.
+			EEditableMeshElementType ElementTypeToSelect = Input.MeshElementsToSelect[ 0 ].ElementAddress.ElementType;
+			for( FMeshElement& MeshElementToSelect : Input.MeshElementsToSelect )
+			{
+				check( MeshElementToSelect.ElementAddress.ElementType == ElementTypeToSelect );
+			}
+
+			if( MeshEditorMode.GetSelectedMeshElementType() != EEditableMeshElementType::Invalid &&
+				MeshEditorMode.GetSelectedMeshElementType() != ElementTypeToSelect )
+			{
+				// We're selecting elements of a different type than we already had selected, so we need to clear our selection first
+				CompoundRevertInput.Subchanges.Add( MoveTemp( FDeselectAllMeshElementsChange( FDeselectAllMeshElementsChangeInput() ).Execute( Object ) ) );
+			}
+
+			for( FMeshElement& MeshElementToSelect : Input.MeshElementsToSelect )
+			{
+				if( MeshElementToSelect.IsValidMeshElement() )
+				{
+					if( CurrentElementSelectionMode == EEditableMeshElementType::Any ||
+					    MeshElementToSelect.ElementAddress.ElementType == CurrentElementSelectionMode )
+					{
+						UEditableMesh* EditableMesh = MeshEditorMode.FindOrCreateEditableMesh( *MeshElementToSelect.Component, MeshElementToSelect.ElementAddress.SubMeshAddress );
+						if( EditableMesh != nullptr )
+						{
+							if( IsElementIDValid( MeshElementToSelect, EditableMesh ) )
+							{
+								FMeshElement& AddedSelectedMeshElement = MeshEditorMode.SelectedMeshElements[ MeshEditorMode.SelectedMeshElements.Add( MeshElementToSelect ) ];
+								AddedSelectedMeshElement.LastSelectTime = CurrentRealTime;
+							}
+						}
+					}
+				}
+			}
+		}
+
+		// Update our transformable list
+		MeshEditorMode.RefreshTransformables();
+	}
+
+	return MakeUnique<FCompoundChange>( MoveTemp( CompoundRevertInput ) );
+}
+
+
+FString FMeshEditorMode::FSelectOrDeselectMeshElementsChange::ToString() const
+{
+	return FString::Printf(
+		TEXT( "Select or Deselect Mesh Elements [MeshElementsToSelect:%s, MeshElementsToDeselect:%s]" ),
+		*LogHelpers::ArrayToString( Input.MeshElementsToSelect ),
+		*LogHelpers::ArrayToString( Input.MeshElementsToDeselect ));
+}
+
+
+TUniquePtr<FChange> FMeshEditorMode::FDeselectAllMeshElementsChange::Execute( UObject* Object )
+{
+	UMeshEditorModeProxyObject* LocalMeshEditorModeProxyObject = CastChecked<UMeshEditorModeProxyObject>( Object );
+	FMeshEditorMode& MeshEditorMode = *LocalMeshEditorModeProxyObject->OwningMeshEditorMode;
+
+	// Back up the current selection so we can restore it on undo
+	FSelectOrDeselectMeshElementsChangeInput RevertInput;
+	RevertInput.MeshElementsToSelect = MeshEditorMode.SelectedMeshElements;
+
+	if( MeshEditorMode.IsActive() )
+	{
+		MeshEditorMode.SelectedMeshElements.Reset();
+		MeshEditorMode.RefreshTransformables();
+	}
+
+	return ( RevertInput.MeshElementsToSelect.Num() > 0 ) ? MakeUnique <FSelectOrDeselectMeshElementsChange>( MoveTemp( RevertInput ) ) : nullptr;
+}
+
+
+FString FMeshEditorMode::FDeselectAllMeshElementsChange::ToString() const
+{
+	return FString::Printf(
+		TEXT( "Deselect All Mesh Elements" ) );
+}
+
+
+TUniquePtr<FChange> FMeshEditorMode::FSetElementSelectionModeChange::Execute( UObject* Object )
+{
+	UMeshEditorModeProxyObject* LocalMeshEditorModeProxyObject = CastChecked<UMeshEditorModeProxyObject>( Object );
+	FMeshEditorMode& MeshEditorMode = *LocalMeshEditorModeProxyObject->OwningMeshEditorMode;
+
+	if( !MeshEditorMode.IsActive() ||
+		Input.Mode == MeshEditorMode.MeshElementSelectionMode )
+	{
+		return nullptr;
+	}
+
+	static TArray<FMeshElement> ElementsToSelect;
+	ElementsToSelect.Empty();
+
+	if( Input.bApplyStoredSelection )
+	{
+		ElementsToSelect = Input.StoredSelection;
+	}
+	else if( GetDefault<UMeshEditorSettings>()->bSeparateSelectionSetPerMode )
+	{
+		// In this mode, the selected elements for each mode are remembered and restored when changing modes
+
+		auto GetSelectedElementsForType = [ &MeshEditorMode ]( EEditableMeshElementType ElementType ) -> TArray<FMeshElement>&
+		{
+			switch( ElementType )
+			{
+				case EEditableMeshElementType::Vertex: return MeshEditorMode.SelectedVertices;
+				case EEditableMeshElementType::Edge: return MeshEditorMode.SelectedEdges;
+				case EEditableMeshElementType::Polygon: return MeshEditorMode.SelectedPolygons;
+				case EEditableMeshElementType::Any: return MeshEditorMode.SelectedMeshElements;
+				default: return MeshEditorMode.SelectedMeshElements;
+			}
+		};
+
+		check( MeshEditorMode.MeshElementSelectionMode != EEditableMeshElementType::Invalid );
+		GetSelectedElementsForType( MeshEditorMode.MeshElementSelectionMode ) = MeshEditorMode.SelectedMeshElements;
+
+		ElementsToSelect = GetSelectedElementsForType( Input.Mode );
+	}
+	else
+	{
+		// In this mode, the current selection is adapted to select related elements of the new type.
+		// e.g. when selecting edge mode, edges of the currently selected polygon or vertices will be selected.
+
+		for( const FMeshElement& MeshElement : MeshEditorMode.SelectedMeshElements )
+		{
+			UPrimitiveComponent* Component = MeshElement.Component.Get();
+			const FEditableMeshElementAddress& ElementAddress = MeshElement.ElementAddress;
+			const FEditableMeshSubMeshAddress& SubMeshAddress = ElementAddress.SubMeshAddress;
+
+			if( Component )
+			{
+				UEditableMesh* EditableMesh = MeshEditorMode.FindOrCreateEditableMesh( *Component, SubMeshAddress );
+
+				if( Input.Mode == EEditableMeshElementType::Vertex )
+				{
+					if( ElementAddress.ElementType == EEditableMeshElementType::Edge )
+					{
+						// Select vertices of the selected edge
+						const FEdgeID EdgeID = FEdgeID( ElementAddress.ElementID );
+						ElementsToSelect.Emplace( Component, SubMeshAddress, EditableMesh->GetEdgeVertex( EdgeID, 0 ) );
+						ElementsToSelect.Emplace( Component, SubMeshAddress, EditableMesh->GetEdgeVertex( EdgeID, 1 ) );
+					}
+					else if( ElementAddress.ElementType == EEditableMeshElementType::Polygon )
+					{
+						// Select vertices of the selected polygon
+						const FPolygonRef PolygonRef = FPolygonRef( ElementAddress.SectionID, FPolygonID( ElementAddress.ElementID ) );
+						const int32 PolygonVertexCount = EditableMesh->GetPolygonPerimeterVertexCount( PolygonRef );
+						for( int32 Index = 0; Index < PolygonVertexCount; ++Index )
+						{
+							ElementsToSelect.Emplace( Component, SubMeshAddress, EditableMesh->GetPolygonPerimeterVertex( PolygonRef, Index ) );
+						}
+					}
+				}
+				else if( Input.Mode == EEditableMeshElementType::Edge )
+				{
+					if( ElementAddress.ElementType == EEditableMeshElementType::Vertex )
+					{
+						// Select edges connected to the selected vertex
+						const FVertexID VertexID = FVertexID( ElementAddress.ElementID );
+						const int32 VertexConnectedEdgeCount = EditableMesh->GetVertexConnectedEdgeCount( VertexID );
+						for( int32 Index = 0; Index < VertexConnectedEdgeCount; ++Index )
+						{
+							ElementsToSelect.Emplace( Component, SubMeshAddress, EditableMesh->GetVertexConnectedEdge( VertexID, Index ) );
+						}
+					}
+					else if( ElementAddress.ElementType == EEditableMeshElementType::Polygon )
+					{
+						// Select edges forming the selected polygon
+						const FPolygonRef PolygonRef = FPolygonRef( ElementAddress.SectionID, FPolygonID( ElementAddress.ElementID ) );
+						const int32 PolygonEdgeCount = EditableMesh->GetPolygonPerimeterEdgeCount( PolygonRef );
+						for( int32 Index = 0; Index < PolygonEdgeCount; ++Index )
+						{
+							bool bEdgeWindingIsReversedForPolygon;
+							ElementsToSelect.Emplace( Component, SubMeshAddress, EditableMesh->GetPolygonPerimeterEdge( PolygonRef, Index, bEdgeWindingIsReversedForPolygon ) );
+						}
+					}
+				}
+				else if( Input.Mode == EEditableMeshElementType::Polygon )
+				{
+					if( ElementAddress.ElementType == EEditableMeshElementType::Vertex )
+					{
+						// Select all polygons containing the selected vertex
+						// @todo mesheditor: is that reasonable? Should it only select a polygon which has all its vertices selected?
+						TArray<FPolygonRef> VertexConnectedPolygons;
+						EditableMesh->GetVertexConnectedPolygons( FVertexID( ElementAddress.ElementID ), VertexConnectedPolygons );
+						for( FPolygonRef VertexConnectedPolygon : VertexConnectedPolygons )
+						{
+							ElementsToSelect.Emplace( Component, SubMeshAddress, VertexConnectedPolygon );
+						}
+					}
+					else if( ElementAddress.ElementType == EEditableMeshElementType::Edge )
+					{
+						// Select all polygons containing the selected edge
+						const FEdgeID EdgeID = FEdgeID( ElementAddress.ElementID );
+						const int32 EdgeConnectedPolygonCount = EditableMesh->GetEdgeConnectedPolygonCount( EdgeID );
+						for( int32 Index = 0; Index < EdgeConnectedPolygonCount; ++Index )
+						{
+							ElementsToSelect.Emplace( Component, SubMeshAddress, EditableMesh->GetEdgeConnectedPolygon( EdgeID, Index ) );
+						}
+					}
+				}
+				else if( Input.Mode == EEditableMeshElementType::Any )
+				{
+					ElementsToSelect = MeshEditorMode.SelectedMeshElements;
+				}
+			}
+		}
+	}
+
+	FSetElementSelectionModeChangeInput RevertInput;
+	RevertInput.Mode = MeshEditorMode.MeshElementSelectionMode;
+	RevertInput.bApplyStoredSelection = true;
+	RevertInput.StoredSelection = MeshEditorMode.SelectedMeshElements;
+
+	// Set new selection mode
+	MeshEditorMode.MeshElementSelectionMode = Input.Mode;
+
+	FSelectOrDeselectMeshElementsChangeInput Select;
+	Select.MeshElementsToSelect = ElementsToSelect;
+	Select.MeshElementsToDeselect = MeshEditorMode.SelectedMeshElements;
+	FSelectOrDeselectMeshElementsChange( MoveTemp( Select ) ).Execute( Object );
+
+	return MakeUnique<FSetElementSelectionModeChange>( RevertInput );
+}
+
+
+FString FMeshEditorMode::FSetElementSelectionModeChange::ToString() const
+{
+	switch( Input.Mode )
+	{
+		case EEditableMeshElementType::Vertex:
+			return TEXT( "Set Vertex Selection Mode" );
+		case EEditableMeshElementType::Edge:
+			return TEXT( "Set Edge Selection Mode" );
+		case EEditableMeshElementType::Polygon:
+			return TEXT( "Set Polygon Selection Mode" );
+		case EEditableMeshElementType::Any:
+			return TEXT( "Set Any Selection Mode" );
+	}
+
+	return FString();
+}
+
+
+FMeshEditorMode::FMeshEditorInteractorData::FMeshEditorInteractorData()
+	: ViewportInteractor( nullptr ),
+	  bGrabberSphereIsValid( false ),
+	  GrabberSphere( 0 ),
+	  bLaserIsValid( false ),
+	  LaserStart( FVector::ZeroVector ),
+	  LaserEnd( FVector::ZeroVector ),
+	  HoveredMeshElement(),
+	  PreviouslyHoveredMeshElement(),
+	  HoverInteractorShape( EInteractorShape::Invalid ),
+	  HoverLocation( FVector::ZeroVector )
+{
+}
+
+FMeshEditorMode::FMeshEditorMode()
+	: HoveredGeometryMaterial( nullptr ),
+	  HoveredFaceMaterial( nullptr ),
+	  HoverFeedbackTimeValue( 0.0 ),
+	  MeshElementSelectionMode( EEditableMeshElementType::Any ),
+	  EquippedVertexAction( EMeshEditAction::Move ),
+	  EquippedEdgeAction( EMeshEditAction::Move ),
+	  EquippedPolygonAction( EMeshEditAction::Move ),
+	  ActiveAction( EMeshEditAction::None ),
+	  ActiveActionInteractor( nullptr ),
+	  bActiveActionNeedsHoverLocation( false ),
+	  bIsFirstActiveActionUpdate( false ),
+	  ExtrudePolygonAxisOrigin( FVector::ZeroVector ),
+	  ExtrudePolygonAxisDirection( FVector::ZeroVector ),
+	  InsetUsingPolygonElement(),
+	  SplitEdgeMeshesAndEdgesToSplit(),
+	  SplitEdgeSplitList(),
+	  EditSharpnessStartLocation( FVector::ZeroVector ),
+	  PreviewRevertChanges(),
+	  PreviewTopologyChange( EMeshTopologyChange::NoTopologyChange ),
+	  MeshEditorModeProxyObject( nullptr ),
+	  SelectingByPaintingRevertChangeInput( nullptr ),
+	  bShowVertexNormals( false ),
+	  bMarqueeSelectTransactionActive( false ),
+	  bShouldFocusToSelection( false ),
+	  bPerInstanceEdits( false ),
+	  AssetContainer( nullptr )
+{
+	AssetContainer = LoadObject<UMeshEditorAssetContainer>(nullptr, TEXT("/Engine/MeshEditor/MeshEditorAssetContainer"));
+	check(AssetContainer != nullptr);
+
+	HoveredGeometryMaterial = AssetContainer->HoveredGeometryMaterial;
+	check( HoveredGeometryMaterial != nullptr );
+
+	HoveredFaceMaterial = AssetContainer->HoveredFaceMaterial;
+	check( HoveredFaceMaterial != nullptr );
+
+	MeshEditorModeProxyObject = NewObject<UMeshEditorModeProxyObject>();
+	MeshEditorModeProxyObject->OwningMeshEditorMode = this;
+
+	// Register mesh editor actions
+	FMeshEditorCommonCommands::Register();
+	FMeshEditorVertexCommands::Register();
+	FMeshEditorEdgeCommands::Register();
+	FMeshEditorPolygonCommands::Register();
+
+	// Register UI commands
+	BindCommands();
+
+	// Notify when the map changes
+	FLevelEditorModule& LevelEditor = FModuleManager::LoadModuleChecked<FLevelEditorModule>( "LevelEditor" );
+	LevelEditor.OnMapChanged().AddRaw( this, &FMeshEditorMode::OnMapChanged );
+	LevelEditor.OnActorSelectionChanged().AddRaw( this, &FMeshEditorMode::OnActorSelectionChanged );
+
+	FEditorDelegates::EndPIE.AddRaw( this, &FMeshEditorMode::OnEndPIE );
+}
+
+
+FMeshEditorMode::~FMeshEditorMode()
+{
+	FEditorDelegates::EndPIE.RemoveAll( this );
+
+	FLevelEditorModule* LevelEditor = FModuleManager::GetModulePtr<FLevelEditorModule>( "LevelEditor" );
+	if( LevelEditor )
+	{
+		LevelEditor->OnActorSelectionChanged().RemoveAll( this );
+		LevelEditor->OnMapChanged().RemoveAll( this );
+	}
+
+	// Unregister mesh editor actions
+	FMeshEditorPolygonCommands::Unregister();
+	FMeshEditorEdgeCommands::Unregister();
+	FMeshEditorVertexCommands::Unregister();
+	FMeshEditorCommonCommands::Unregister();
+
+	MeshEditorModeProxyObject = nullptr;
+	AssetContainer = nullptr;
+}
+
+
+void FMeshEditorMode::OnMapChanged( UWorld* World, EMapChangeType MapChangeType )
+{
+	if( MapChangeType == EMapChangeType::TearDownWorld )
+	{
+		RemoveEditableMeshReferences();
+	}
+}
+
+
+void FMeshEditorMode::OnEndPIE( bool bIsSimulating )
+{
+	if( bIsSimulating )
+	{
+		RemoveEditableMeshReferences();
+	}
+}
+
+
+void FMeshEditorMode::RemoveEditableMeshReferences()
+{
+	// Instanced meshes live within the level itself. So remove all possible references to any editable mesh when the map changes,
+	// to prevent unreachable paths following GC.
+	CachedEditableMeshes.Empty();
+	SelectedMeshElements.Empty();
+	SelectedVertices.Empty();
+	SelectedEdges.Empty();
+	SelectedPolygons.Empty();
+	PreviewRevertChanges.Empty();
+
+	if( ViewportWorldInteraction != nullptr )
+	{
+		RefreshTransformables();
+	}
+}
+
+void FMeshEditorMode::PlayStartActionSound(EMeshEditAction NewAction, UViewportInteractor* ActionInteractor)
+{
+	if (ActionInteractor != nullptr)
+	{
+		UGameplayStatics::PlaySoundAtLocation(GetWorld(), AssetContainer->DefaultSound,
+			ActionInteractor->GetTransform().GetLocation(), FRotator::ZeroRotator,
+			MeshEd::SFXMultiplier->GetFloat());
+	}
+	else
+	{
+		UGameplayStatics::PlaySound2D(GetWorld(), AssetContainer->DefaultSound, MeshEd::SFXMultiplier->GetFloat());
+	}
+}
+
+void FMeshEditorMode::PlayFinishActionSound(EMeshEditAction NewAction, UViewportInteractor* ActionInteractor)
+{
+	if (ActionInteractor != nullptr)
+	{
+		UGameplayStatics::PlaySoundAtLocation(GetWorld(), AssetContainer->DefaultSound,
+			ActionInteractor->GetTransform().GetLocation(), FRotator::ZeroRotator, 
+			0.5f);
+	}
+	else
+	{
+		UGameplayStatics::PlaySound2D(GetWorld(), AssetContainer->DefaultSound, 1.5f);
+	}
+}
+
+void FMeshEditorMode::BindCommands()
+{
+	const FMeshEditorCommonCommands& MeshEditorCommonActions( FMeshEditorCommonCommands::Get() );
+	const FMeshEditorVertexCommands& MeshEditorVertexActions( FMeshEditorVertexCommands::Get() );
+	const FMeshEditorEdgeCommands& MeshEditorEdgeActions( FMeshEditorEdgeCommands::Get() );
+	const FMeshEditorPolygonCommands& MeshEditorPolygonActions( FMeshEditorPolygonCommands::Get() );
+
+	// Register editing modes (equipped actions)
+	RegisterVertexEditingMode( MeshEditorVertexActions.MoveVertex, EMeshEditAction::Move );
+	RegisterVertexEditingMode( MeshEditorVertexActions.ExtendVertex, EMeshEditAction::ExtendVertex );
+	RegisterVertexEditingMode( MeshEditorVertexActions.EditVertexCornerSharpness, EMeshEditAction::EditVertexCornerSharpness );
+
+	RegisterEdgeEditingMode( MeshEditorEdgeActions.MoveEdge, EMeshEditAction::Move );
+	RegisterEdgeEditingMode( MeshEditorEdgeActions.SplitEdge, EMeshEditAction::SplitEdge );
+	RegisterEdgeEditingMode( MeshEditorEdgeActions.SplitEdgeAndDragVertex, EMeshEditAction::SplitEdgeAndDragVertex );
+	RegisterEdgeEditingMode( MeshEditorEdgeActions.InsertEdgeLoop, EMeshEditAction::InsertEdgeLoop );
+	RegisterEdgeEditingMode( MeshEditorEdgeActions.ExtendEdge, EMeshEditAction::ExtendEdge );
+	RegisterEdgeEditingMode( MeshEditorEdgeActions.EditEdgeCreaseSharpness, EMeshEditAction::EditEdgeCreaseSharpness );
+
+	RegisterPolygonEditingMode( MeshEditorPolygonActions.MovePolygon, EMeshEditAction::Move );
+	RegisterPolygonEditingMode( MeshEditorPolygonActions.ExtrudePolygon, EMeshEditAction::ExtrudePolygon );
+//	RegisterPolygonEditingMode( MeshEditorPolygonActions.FreelyExtrudePolygon, EMeshEditAction::FreelyExtrudePolygon );	// @todo mesheditor: "Freely Extrude" disabled until we have proper gizmo/snap support
+	RegisterPolygonEditingMode( MeshEditorPolygonActions.InsetPolygon, EMeshEditAction::InsetPolygon );
+	RegisterPolygonEditingMode( MeshEditorPolygonActions.BevelPolygon, EMeshEditAction::BevelPolygon );
+
+	RegisterCommonEditingMode( MeshEditorCommonActions.DrawVertices, EMeshEditAction::DrawVertices );
+
+	// Register commands which work regardless of which element type is selected
+	RegisterCommand( MeshEditorCommonActions.DeleteMeshElement, FExecuteAction::CreateLambda( [this] { DeleteSelectedMeshElement(); } ) );
+	RegisterCommand( MeshEditorCommonActions.AddSubdivisionLevel, FExecuteAction::CreateLambda( [this] { AddOrRemoveSubdivisionLevel( true ); } ) );
+	RegisterCommand( MeshEditorCommonActions.RemoveSubdivisionLevel, FExecuteAction::CreateLambda( [this] { AddOrRemoveSubdivisionLevel( false ); } ) );
+
+	// @todo mesheditor: support EUserInterfaceActionType::ToggleButton actions in the UI, and extend RegisterCommand to allow
+	// a delegate returning check state.
+	RegisterCommand( MeshEditorCommonActions.ShowVertexNormals, FExecuteAction::CreateLambda( [this] { bShowVertexNormals = !bShowVertexNormals; } ) );
+
+	RegisterCommand( MeshEditorCommonActions.MarqueeSelectVertices, FExecuteAction::CreateLambda( [this] { PerformMarqueeSelect( EEditableMeshElementType::Vertex ); } ) );
+	RegisterCommand( MeshEditorCommonActions.MarqueeSelectEdges, FExecuteAction::CreateLambda( [this] { PerformMarqueeSelect( EEditableMeshElementType::Edge ); } ) );
+	RegisterCommand( MeshEditorCommonActions.MarqueeSelectPolygons, FExecuteAction::CreateLambda( [this] { PerformMarqueeSelect( EEditableMeshElementType::Polygon ); } ) );
+	RegisterCommand( MeshEditorCommonActions.FrameSelectedElements, FExecuteAction::CreateLambda( [this] { bShouldFocusToSelection = true; } ) );
+
+	RegisterCommand( MeshEditorCommonActions.SetVertexSelectionMode, FExecuteAction::CreateLambda( [this] { SetMeshElementSelectionMode( EEditableMeshElementType::Vertex ); } ) );
+	RegisterCommand( MeshEditorCommonActions.SetEdgeSelectionMode, FExecuteAction::CreateLambda( [this] { SetMeshElementSelectionMode( EEditableMeshElementType::Edge ); } ) );
+	RegisterCommand( MeshEditorCommonActions.SetPolygonSelectionMode, FExecuteAction::CreateLambda( [this] { SetMeshElementSelectionMode( EEditableMeshElementType::Polygon ); } ) );
+	RegisterCommand( MeshEditorCommonActions.SetAnySelectionMode, FExecuteAction::CreateLambda( [this] { SetMeshElementSelectionMode( EEditableMeshElementType::Any ); } ) );
+
+	RegisterCommand( MeshEditorCommonActions.QuadrangulateMesh, FExecuteAction::CreateLambda( [ this ] { QuadrangulateMesh(); } ) );
+
+	// Register element-specific commands
+	RegisterVertexCommand( MeshEditorVertexActions.RemoveVertex, FExecuteAction::CreateLambda( [this] { RemoveSelectedVertices(); } ) );
+	RegisterVertexCommand( MeshEditorVertexActions.WeldVertices, FExecuteAction::CreateLambda( [this] { WeldSelectedVertices(); } ) );
+
+	RegisterEdgeCommand( MeshEditorEdgeActions.RemoveEdge, FExecuteAction::CreateLambda( [this] { RemoveSelectedEdges(); } ) );
+	RegisterEdgeCommand( MeshEditorEdgeActions.SoftenEdge, FExecuteAction::CreateLambda( [this] { MakeSelectedEdgesHardOrSoft( false ); } ) );
+	RegisterEdgeCommand( MeshEditorEdgeActions.HardenEdge, FExecuteAction::CreateLambda( [this] { MakeSelectedEdgesHardOrSoft( true ); } ) );
+	RegisterEdgeCommand( MeshEditorEdgeActions.SelectEdgeLoop, FExecuteAction::CreateLambda( [ this ] { SelectEdgeLoops(); } ) );
+
+	RegisterPolygonCommand( MeshEditorPolygonActions.FlipPolygon, FExecuteAction::CreateLambda( [this] { FlipSelectedPolygons(); } ) );
+	RegisterPolygonCommand( MeshEditorPolygonActions.TriangulatePolygon, FExecuteAction::CreateLambda( [this] { TriangulateSelectedPolygons(); } ) );
+	RegisterPolygonCommand( MeshEditorPolygonActions.TessellatePolygon, FExecuteAction::CreateLambda( [ this ] { TessellateSelectedPolygons(); } ) );
+	RegisterPolygonCommand( MeshEditorPolygonActions.AssignMaterial, FExecuteAction::CreateLambda( [this] { AssignSelectedMaterialToSelectedPolygons(); } ) );
+
+	// Bind common actions
+	CommonCommands = MakeShared<FUICommandList>();
+	for( const TTuple<TSharedPtr<FUICommandInfo>, FUIAction>& CommonAction : CommonActions )
+	{
+		CommonCommands->MapAction( CommonAction.Get<0>(), CommonAction.Get<1>() );
+	}
+
+	// Bind vertex actions
+	VertexCommands = MakeShared<FUICommandList>();
+	for( const TTuple<TSharedPtr<FUICommandInfo>, FUIAction>& VertexAction : VertexActions )
+	{
+		VertexCommands->MapAction( VertexAction.Get<0>(), VertexAction.Get<1>() );
+	}
+
+	// Bind edge actions
+	EdgeCommands = MakeShared<FUICommandList>();
+	for( const TTuple<TSharedPtr<FUICommandInfo>, FUIAction>& EdgeAction : EdgeActions )
+	{
+		EdgeCommands->MapAction( EdgeAction.Get<0>(), EdgeAction.Get<1>() );
+	}
+
+	// Bind polygon actions
+	PolygonCommands = MakeShared<FUICommandList>();
+	for( const TTuple<TSharedPtr<FUICommandInfo>, FUIAction>& PolygonAction : PolygonActions )
+	{
+		PolygonCommands->MapAction( PolygonAction.Get<0>(), PolygonAction.Get<1>() );
+	}
+}
+
+
+void FMeshEditorMode::RegisterCommonEditingMode( const TSharedPtr<FUICommandInfo>& Command, EMeshEditAction EditingMode )
+{
+	RegisterVertexEditingMode( Command, EditingMode );
+	RegisterEdgeEditingMode( Command, EditingMode );
+	RegisterPolygonEditingMode( Command, EditingMode );
+}
+
+void FMeshEditorMode::RegisterVertexEditingMode( const TSharedPtr<FUICommandInfo>& Command, EMeshEditAction EditingMode )
+{
+	VertexActions.Emplace( Command, FUIAction(
+		FExecuteAction::CreateLambda( [this, EditingMode] { EquippedVertexAction = EditingMode; } ),
+		FCanExecuteAction::CreateLambda( [this] { return IsMeshElementTypeSelectedOrIsActiveSelectionMode( EEditableMeshElementType::Vertex ); } ),
+		FIsActionChecked::CreateLambda( [this, EditingMode] { return ( EquippedVertexAction == EditingMode ); } )
+	) );
+}
+
+
+void FMeshEditorMode::RegisterEdgeEditingMode( const TSharedPtr<FUICommandInfo>& Command, EMeshEditAction EditingMode )
+{
+	EdgeActions.Emplace( Command, FUIAction(
+		FExecuteAction::CreateLambda( [this, EditingMode] { EquippedEdgeAction = EditingMode; } ),
+		FCanExecuteAction::CreateLambda( [this] { return IsMeshElementTypeSelectedOrIsActiveSelectionMode( EEditableMeshElementType::Edge ); } ),
+		FIsActionChecked::CreateLambda( [this, EditingMode] { return ( EquippedEdgeAction == EditingMode ); } )
+	) );
+}
+
+
+void FMeshEditorMode::RegisterPolygonEditingMode( const TSharedPtr<FUICommandInfo>& Command, EMeshEditAction EditingMode )
+{
+	PolygonActions.Emplace( Command, FUIAction(
+		FExecuteAction::CreateLambda( [this, EditingMode] { EquippedPolygonAction = EditingMode; } ),
+		FCanExecuteAction::CreateLambda( [this] { return IsMeshElementTypeSelectedOrIsActiveSelectionMode( EEditableMeshElementType::Polygon ); } ),
+		FIsActionChecked::CreateLambda( [this, EditingMode] { return ( EquippedPolygonAction == EditingMode ); } )
+	) );
+}
+
+
+void FMeshEditorMode::RegisterCommand( const TSharedPtr<FUICommandInfo>& Command, const FExecuteAction& ExecuteAction )
+{
+	CommonActions.Emplace( Command, FUIAction( ExecuteAction, FCanExecuteAction::CreateLambda( [this] { return GetSelectedMeshElementType() != EEditableMeshElementType::Invalid;} ) ) );
+	VertexActions.Emplace( Command, FUIAction( ExecuteAction, FCanExecuteAction::CreateLambda( [this] { return GetSelectedMeshElementType() != EEditableMeshElementType::Invalid;} ) ) );
+	EdgeActions.Emplace( Command, FUIAction( ExecuteAction, FCanExecuteAction::CreateLambda( [this] { return GetSelectedMeshElementType() != EEditableMeshElementType::Invalid;} ) ) );
+	PolygonActions.Emplace( Command, FUIAction( ExecuteAction, FCanExecuteAction::CreateLambda( [this] { return GetSelectedMeshElementType() != EEditableMeshElementType::Invalid;} ) ) );
+}
+
+
+void FMeshEditorMode::RegisterVertexCommand( const TSharedPtr<FUICommandInfo>& Command, const FExecuteAction& ExecuteAction )
+{
+	VertexActions.Emplace( Command, FUIAction( ExecuteAction, FCanExecuteAction::CreateLambda( [this] { return IsMeshElementTypeSelected( EEditableMeshElementType::Vertex ); } ) ) );
+}
+
+
+void FMeshEditorMode::RegisterEdgeCommand( const TSharedPtr<FUICommandInfo>& Command, const FExecuteAction& ExecuteAction )
+{
+	EdgeActions.Emplace( Command, FUIAction( ExecuteAction, FCanExecuteAction::CreateLambda( [this] { return IsMeshElementTypeSelected( EEditableMeshElementType::Edge ); } ) ) );
+}
+
+
+void FMeshEditorMode::RegisterPolygonCommand( const TSharedPtr<FUICommandInfo>& Command, const FExecuteAction& ExecuteAction )
+{
+	PolygonActions.Emplace( Command, FUIAction( ExecuteAction, FCanExecuteAction::CreateLambda( [this] { return IsMeshElementTypeSelected( EEditableMeshElementType::Polygon ); } ) ) );
+}
+
+
+void FMeshEditorMode::Enter()
+{
+	// Call parent implementation
+	FEdMode::Enter();
+
+	UEditorWorldExtensionCollection* ExtensionCollection = GEditor->GetEditorWorldExtensionsManager()->GetEditorWorldExtensions( GetWorld() );
+	check( ExtensionCollection != nullptr );
+	this->ViewportWorldInteraction = CastChecked<UViewportWorldInteraction>( ExtensionCollection->AddExtension( UViewportWorldInteraction::StaticClass() ) );
+	check( ViewportWorldInteraction != nullptr );
+
+	// Register to find out about viewport interaction events
+	ViewportWorldInteraction->OnViewportInteractionHoverUpdate().AddRaw( this, &FMeshEditorMode::OnViewportInteractionHoverUpdate );
+	ViewportWorldInteraction->OnViewportInteractionInputAction().AddRaw( this, &FMeshEditorMode::OnViewportInteractionInputAction );
+	ViewportWorldInteraction->OnViewportInteractionInputUnhandled().AddRaw( this, &FMeshEditorMode::OnViewportInteractionInputUnhandled );
+	ViewportWorldInteraction->OnStartDragging().AddRaw( this, &FMeshEditorMode::OnViewportInteractionStartDragging );
+	ViewportWorldInteraction->OnStopDragging().AddRaw( this, &FMeshEditorMode::OnViewportInteractionStopDragging );
+	ViewportWorldInteraction->OnFinishedMovingTransformables().AddRaw( this, &FMeshEditorMode::OnViewportInteractionFinishedMovingTransformables );
+
+	// @todo gizmo: Forcing "All" mode when you enter Geometry mode, because it is generally the most useful.
+	ViewportWorldInteraction->SetGizmoHandleType( EGizmoHandleTypes::All );
+
+	// Register our system for transforming mesh elements
+	UMeshElementTransformer* MeshElementTransformer = NewObject<UMeshElementTransformer>();
+	ViewportWorldInteraction->SetTransformer( MeshElementTransformer );
+
+	this->VREditorMode = Cast<UVREditorMode>( ExtensionCollection->FindExtension( UVREditorMode::StaticClass() ) );
+	if( VREditorMode && VREditorMode->IsFullyInitialized() )
+	{
+		VREditorMode->OnPlaceDraggedMaterial().AddRaw( this, &FMeshEditorMode::OnVREditorModePlaceDraggedMaterial );
+	}
+
+	// Add toolkit
+	if( !Toolkit.IsValid() )
+	{
+		Toolkit = MakeShared< FMeshEditorModeToolkit >( *this );
+		Toolkit->Init( Owner->GetToolkitHost() );
+	}
+
+	RegisterWithExternalMenuSystem();
+}
+
+
+void FMeshEditorMode::Exit()
+{
+	UnregisterWithExternalMenuSystem();
+	if( VREditorMode && VREditorMode->IsFullyInitialized() )
+	{
+		VREditorMode->OnPlaceDraggedMaterial().RemoveAll( this );
+	}
+
+	// If anything is selected, go ahead and deselect everything now
+	if( SelectedMeshElements.Num() > 0 )
+	{
+		const FScopedTransaction Transaction( LOCTEXT( "UndoDeselectingAllMeshElements", "Deselect All Elements" ) );
+
+		StoreUndo( MeshEditorModeProxyObject, FDeselectAllMeshElementsChange( MoveTemp( FDeselectAllMeshElementsChangeInput() ) ).Execute( MeshEditorModeProxyObject ) );
+	}
+
+	FToolkitManager::Get().CloseToolkit( Toolkit.ToSharedRef() );
+	Toolkit.Reset();
+
+	// Unregister from event handlers
+	if( IViewportInteractionModule::IsAvailable() )
+	{
+		if( ViewportWorldInteraction != nullptr )
+		{
+			// Make sure gizmo is visible.  We may have hidden it
+			ViewportWorldInteraction->SetTransformGizmoVisible( true );
+
+			// Unregister mesh element transformer
+			ViewportWorldInteraction->SetTransformer( nullptr );
+
+			ViewportWorldInteraction->OnStartDragging().RemoveAll( this );
+			ViewportWorldInteraction->OnStopDragging().RemoveAll( this );
+			ViewportWorldInteraction->OnFinishedMovingTransformables().RemoveAll( this );
+			ViewportWorldInteraction->OnViewportInteractionHoverUpdate().RemoveAll( this );
+			ViewportWorldInteraction->OnViewportInteractionInputAction().RemoveAll( this );
+			ViewportWorldInteraction->OnViewportInteractionInputUnhandled().RemoveAll( this );
+
+			UEditorWorldExtensionCollection* ExtensionCollection = GEditor->GetEditorWorldExtensionsManager()->GetEditorWorldExtensions( GetWorld() );
+			if( ExtensionCollection != nullptr )
+			{
+				ExtensionCollection->RemoveExtension( ViewportWorldInteraction );
+			}
+
+			ViewportWorldInteraction = nullptr;
+		}
+	}
+
+	// Geometry will no longer be selected, so notify that selection changed.  This makes sure that other modes are prepared
+	// to interact with whichever objects are still selected, now that mesh editing has finished
+	if( !GIsRequestingExit )
+	{
+		GEditor->NoteSelectionChange();
+	}
+
+	// Call parent implementation
+	FEdMode::Exit();
+}
+
+
+UEditableMesh* FMeshEditorMode::FindOrCreateEditableMesh( UPrimitiveComponent& Component, const FEditableMeshSubMeshAddress& SubMeshAddress )
+{
+	UEditableMesh* EditableMesh = nullptr;
+
+	// Grab the existing editable mesh from our cache if we have one, otherwise create one now
+	UEditableMesh** EditableMeshPtr = CachedEditableMeshes.Find( SubMeshAddress );
+	if( EditableMeshPtr != nullptr )
+	{
+		EditableMesh = *EditableMeshPtr;
+	}
+	else
+	{
+		if( SubMeshAddress.EditableMeshFormat != nullptr )
+		{
+			// @todo mesheditor perf: This is going to HITCH as you hover over meshes.  Ideally we do this on a thread, or worst case give the user a progress dialog.  Maybe save out the editable mesh in editor builds?
+			EditableMesh = UEditableMeshFactory::MakeEditableMesh( &Component, SubMeshAddress );
+			CachedEditableMeshes.Add( SubMeshAddress, EditableMesh );
+
+			if( GetDefault<UMeshEditorSettings>()->bAutoQuadrangulate )
+			{
+				static TArray<FPolygonRef> NewPolygonsRefs;
+				EditableMesh->StartModification( EMeshModificationType::Final, EMeshTopologyChange::TopologyChange );
+				EditableMesh->QuadrangulateMesh( NewPolygonsRefs );
+				EditableMesh->EndModification();
+			}
+
+			// Enable undo tracking on this mesh
+			EditableMesh->SetAllowUndo( true );
+		}
+	}
+
+	return EditableMesh;
+}
+
+
+void FMeshEditorMode::RollbackPreviewChanges()
+{
+	for( TTuple<UObject*, TUniquePtr<FChange>>& ObjectAndPreviewRevertChange : PreviewRevertChanges )
+	{
+		UObject* Object = ObjectAndPreviewRevertChange.Get<0>();
+		TUniquePtr<FChange>& PreviewRevertChange = ObjectAndPreviewRevertChange.Get<1>();
+
+		UEditableMesh* EditableMesh = Cast<UEditableMesh>( Object );
+
+		// @todo mesheditor perf: When rolling back a preview change right before applying an edit in the same frame, we might be able 
+		// to skip certain parts of the update (such as subdivision geometry refresh).  This should get us better performance!
+
+		// @todo mesheditor debug
+		//GWarn->Logf( TEXT( "---------- Rolling Back Preview Change (Object:%s) ----------" ), *Object->GetName() );
+		//PreviewRevertChange->PrintToLog( *GWarn );
+
+		TUniquePtr<FChange> UnusedChangeToUndoRevert = MoveTemp( PreviewRevertChange->Execute( Object ) );
+
+		// @todo mesheditor debug
+		//GWarn->Logf( TEXT( "-----(Here's what the Redo looks like)-----" ) );
+		//UnusedChangeToUndoRevert->PrintToLog( *GWarn );
+		//GWarn->Logf( TEXT( "-------------------------------------------" ) );
+	}
+	PreviewRevertChanges.Reset();
+}
+
+
+void FMeshEditorMode::Tick( FEditorViewportClient* ViewportClient, float DeltaTime )
+{
+	// Call parent implementation
+	FEdMode::Tick( ViewportClient, DeltaTime );
+
+
+	// Roll back whatever we changed last time while previewing.  We need the selected mesh elements to match
+	// the mesh before any temporary changes were made.
+	RollbackPreviewChanges();
+
+	// Update the cached view location
+	UpdateCameraToWorldTransform( *ViewportClient );
+
+	if( bShouldFocusToSelection )
+	{
+		bShouldFocusToSelection = false;
+		FrameSelectedElements( ViewportClient );
+	}
+
+	// @todo mesheditor: Should take into account world scaling while in VR (room space interactor movement threshold)
+	const float MinDeltaForInertialMovement = MeshEd::MinDeltaForInertialMovement->GetFloat();	// cm/frame
+
+
+	// If we're currently selecting elements by painting, go ahead and do that now
+	if( ActiveAction == EMeshEditAction::SelectByPainting )
+	{
+		FMeshEditorInteractorData& MeshEditorInteractorData = GetMeshEditorInteractorData( ActiveActionInteractor );
+
+		// If not already selected, add it to our selection set
+		if( GetSelectedMeshElementIndex( MeshEditorInteractorData.HoveredMeshElement ) == INDEX_NONE )
+		{
+			// Only add elements of the same type.  Otherwise it would just cause things to become deselected as you move between
+			// different element types, as we don't allow you to select elements that have overlapping geometry
+			if( GetSelectedMeshElementType() == EEditableMeshElementType::Invalid || 
+				GetSelectedMeshElementType() == MeshEditorInteractorData.HoveredMeshElement.ElementAddress.ElementType ) 
+			{
+				FSelectOrDeselectMeshElementsChangeInput ChangeInput;
+
+				// Select the element under the mouse cursor
+				ChangeInput.MeshElementsToSelect.Add( MeshEditorInteractorData.HoveredMeshElement );
+
+				check( MeshEditorModeProxyObject != nullptr );
+				SelectingByPaintingRevertChangeInput->Subchanges.Add( FSelectOrDeselectMeshElementsChange( MoveTemp( ChangeInput ) ).Execute( MeshEditorModeProxyObject ) );
+			}
+		}
+	}
+
+
+	// Expire any fully faded out hovered elements
+	{
+		const double CurrentRealTime = FSlateApplication::Get().GetCurrentTime();
+
+		const float HoverFadeTime = MeshEd::HoverFadeDuration->GetFloat();
+		for( int32 ElementIndex = 0; ElementIndex < FadingOutHoveredMeshElements.Num(); ++ElementIndex )
+		{
+			const FMeshElement& ExistingElement = FadingOutHoveredMeshElements[ ElementIndex ];
+			const float TimeSinceLastHovered = CurrentRealTime - ExistingElement.LastHoverTime;
+			if( TimeSinceLastHovered >= HoverFadeTime )
+			{
+				FadingOutHoveredMeshElements.RemoveAtSwap( ElementIndex-- );
+			}
+		}
+	}
+
+	if( ActiveAction != EMeshEditAction::None &&
+		ActiveAction != EMeshEditAction::SelectByPainting )		// When selecting, no updates are needed
+	{
+		const bool bIsActionFinishing = false;
+		UpdateActiveAction( bIsActionFinishing );
+	}
+
+
+	// Advanced hover feedback time
+	HoverFeedbackTimeValue += DeltaTime;
+
+	// End the marquee select transaction if necessary
+	if( MarqueeSelectTransaction.IsValid() && !bMarqueeSelectTransactionActive )
+	{
+		MarqueeSelectTransaction.Reset();
+	}
+
+
+	// Clear hover for all interactors.  We'll re-detect what's hovered every frame when OnViewportInteractionHoverUpdate() is called.
+	for( FMeshEditorInteractorData& MeshEditorInteractorData : MeshEditorInteractorDatas )
+	{
+		MeshEditorInteractorData.PreviouslyHoveredMeshElement = MeshEditorInteractorData.HoveredMeshElement;
+		MeshEditorInteractorData.HoveredMeshElement = FMeshElement();
+		MeshEditorInteractorData.HoverLocation = FVector::ZeroVector;
+	}
+
+
+	// Hide the transform gizmo while we're doing things.  It actually will get in the way of our hit tests!
+	{
+		const EEditableMeshElementType SelectedMeshElementType = GetSelectedMeshElementType();
+		ViewportWorldInteraction->SetTransformGizmoVisible(
+			( ActiveAction == EMeshEditAction::MoveUsingGizmo ) ||
+			( ActiveAction == EMeshEditAction::None &&
+				( ( EquippedPolygonAction == EMeshEditAction::Move && SelectedMeshElementType == EEditableMeshElementType::Polygon ) ||
+				  ( EquippedVertexAction == EMeshEditAction::Move && SelectedMeshElementType == EEditableMeshElementType::Vertex ) ||
+				  ( EquippedEdgeAction == EMeshEditAction::Move && SelectedMeshElementType == EEditableMeshElementType::Edge ) ) ) 
+			);
+	}
+}
+
+
+void FMeshEditorMode::UpdateCameraToWorldTransform( const FEditorViewportClient& ViewportClient )
+{
+	// Get it from VR head position, if valid
+	if( ViewportWorldInteraction->HaveHeadTransform() )
+	{
+		CachedCameraToWorld = ViewportWorldInteraction->GetHeadTransform();
+	}
+	else
+	{
+		CachedCameraToWorld = FTransform( ViewportClient.GetViewTransform().GetRotation(), ViewportClient.GetViewTransform().GetLocation() );
+	}
+}
+
+
+bool FMeshEditorMode::InputKey( FEditorViewportClient* ViewportClient, FViewport* Viewport, FKey Key, EInputEvent Event )
+{
+	bool bHandled = false;
+
+	// Roll back whatever we changed last time while previewing.  We need the selected mesh elements to match
+	// the mesh before any temporary changes were made.
+	RollbackPreviewChanges();
+
+	// If there is still a marquee select transaction pending completion since the last drag operation,
+	// end it here (prior to the next drag operation potentially starting).
+	if( MarqueeSelectTransaction.IsValid() )
+	{
+		MarqueeSelectTransaction.Reset();
+		bMarqueeSelectTransactionActive = false;
+	}
+
+	if( Event == IE_Pressed )
+	{
+		const FUICommandList* CommandList = GetCommandListForSelectedElementType();
+		check( CommandList );
+
+		FModifierKeysState ModifierKeysState = FSlateApplication::Get().GetModifierKeys();
+		bHandled = CommandList->ProcessCommandBindings( Key, ModifierKeysState, false );
+	}
+
+	else if( Event == IE_DoubleClick )
+	{
+		// Absorb double clicks.  Otherwise they'll select actors/components while editing geometry.
+		bHandled = true;
+	}
+
+	return bHandled ? true : FEdMode::InputKey( ViewportClient, Viewport, Key, Event );
+}
+
+
+const FUICommandList* FMeshEditorMode::GetCommandListForSelectedElementType() const
+{
+	switch( GetSelectedMeshElementType() )
+	{
+		case EEditableMeshElementType::Vertex:
+			return VertexCommands.Get();
+
+		case EEditableMeshElementType::Edge:
+			return EdgeCommands.Get();
+
+		case EEditableMeshElementType::Polygon:
+			return PolygonCommands.Get();
+
+		default:
+			return CommonCommands.Get();
+	}
+}
+
+
+void FMeshEditorMode::CommitEditableMeshIfNecessary( UEditableMesh* EditableMesh, UPrimitiveComponent* Component )
+{
+	if( bPerInstanceEdits && !EditableMesh->IsCommittedAsInstance() )
+	{
+		UEditableMesh* NewEditableMesh = EditableMesh->CommitInstance( Component );
+		NewEditableMesh->SetAllowUndo(true);
+
+		// Commit the editable mesh as a new instance in the static mesh component
+		const FEditableMeshSubMeshAddress& OldSubMeshAddress = EditableMesh->GetSubMeshAddress();
+		const FEditableMeshSubMeshAddress& NewSubMeshAddress = NewEditableMesh->GetSubMeshAddress();
+
+		CachedEditableMeshes.Add( NewSubMeshAddress, NewEditableMesh );
+
+		auto FixUpMeshElement = [ &OldSubMeshAddress, &NewSubMeshAddress, Component ]( FMeshElement& MeshElement )
+		{
+			if( MeshElement.Component.Get() == Component &&
+			   MeshElement.ElementAddress.SubMeshAddress == OldSubMeshAddress )
+			{
+				MeshElement.ElementAddress.SubMeshAddress = NewSubMeshAddress;
+			}
+		};
+
+		auto FixUpMeshElements = [ &OldSubMeshAddress, &NewSubMeshAddress, Component, &FixUpMeshElement ]( TArray<FMeshElement>& MeshElements )
+		{
+			for( FMeshElement& MeshElement : MeshElements )
+			{
+				FixUpMeshElement( MeshElement );
+			}
+		};
+
+		// Change selection as an undoable transaction
+		FSelectOrDeselectMeshElementsChangeInput ChangeInput;
+		for( const FMeshElement& MeshElement : SelectedMeshElements )
+		{
+			if( MeshElement.ElementAddress.SubMeshAddress == OldSubMeshAddress )
+			{
+				ChangeInput.MeshElementsToDeselect.Add( MeshElement );
+				FMeshElement NewMeshElement = MeshElement;
+				NewMeshElement.ElementAddress.SubMeshAddress = NewSubMeshAddress;
+				ChangeInput.MeshElementsToSelect.Add( NewMeshElement );
+			}
+		}
+		StoreUndo( MeshEditorModeProxyObject, FSelectOrDeselectMeshElementsChange( MoveTemp( ChangeInput ) ).Execute( MeshEditorModeProxyObject ) );
+
+		FixUpMeshElements( SelectedVertices );
+		FixUpMeshElements( SelectedEdges );
+		FixUpMeshElements( SelectedPolygons );
+		FixUpMeshElements( FadingOutHoveredMeshElements );
+
+		for( FMeshEditorInteractorData& MeshEditorInteractorData : MeshEditorInteractorDatas )
+		{
+			FixUpMeshElement( MeshEditorInteractorData.HoveredMeshElement );
+			FixUpMeshElement( MeshEditorInteractorData.PreviouslyHoveredMeshElement );
+		}
+
+		// @todo mesheditor: this is a little bit fragile. Ideally we initialize these things after the new instance has been created.
+		FixUpMeshElement( InsetUsingPolygonElement );
+
+		for( auto& SplitEdgeMeshAndEdgesToSplit : SplitEdgeMeshesAndEdgesToSplit )
+		{
+			for( auto& EdgeToSplit : SplitEdgeMeshAndEdgesToSplit.Value )
+			{
+				FixUpMeshElement( EdgeToSplit.Get<0>() );
+			}
+		}
+
+		RefreshTransformables();
+	}
+	else if( !EditableMesh->IsCommitted() )
+	{
+		EditableMesh->Commit();
+	}
+}
+
+
+void FMeshEditorMode::CommitSelectedMeshes()
+{
+	// @todo mesheditor: this is inefficient as it iterates the mesh elements arrays each time round the loop to fix up FMeshElements
+	for( int32 SelectedElementIndex = 0; SelectedElementIndex < this->SelectedMeshElements.Num(); ++SelectedElementIndex )
+	{
+		const FMeshElement& SelectedMeshElement = this->SelectedMeshElements[ SelectedElementIndex ];
+		if( SelectedMeshElement.IsValidMeshElement() )
+		{
+			UEditableMesh* EditableMesh = this->FindOrCreateEditableMesh( *SelectedMeshElement.Component, SelectedMeshElement.ElementAddress.SubMeshAddress );
+			if( EditableMesh != nullptr )
+			{
+				CommitEditableMeshIfNecessary( EditableMesh, SelectedMeshElement.Component.Get() );
+			}
+		}
+	}
+}
+
+
+void FMeshEditorMode::PropagateInstanceChanges()
+{
+	for( int32 SelectedElementIndex = 0; SelectedElementIndex < SelectedMeshElements.Num(); ++SelectedElementIndex )
+	{
+		const FMeshElement& SelectedMeshElement = SelectedMeshElements[ SelectedElementIndex ];
+		if( SelectedMeshElement.IsValidMeshElement() )
+		{
+			UEditableMesh* EditableMesh = FindOrCreateEditableMesh( *SelectedMeshElement.Component, SelectedMeshElement.ElementAddress.SubMeshAddress );
+			if( EditableMesh != nullptr )
+			{
+				EditableMesh->PropagateInstanceChanges();
+			}
+		}
+	}
+
+	CachedEditableMeshes.Empty();
+}
+
+
+bool FMeshEditorMode::CanPropagateInstanceChanges() const
+{
+	// @todo mesheditor: this could be more thorough:
+	// it should not allow instance changes to be propagated if more than one instance is selected which derives from the same static mesh.
+	// However MeshEditorMode has no generic way to know if this is the case (and it's unclear how that might be presented in the API).
+	for( int32 SelectedElementIndex = 0; SelectedElementIndex < SelectedMeshElements.Num(); ++SelectedElementIndex )
+	{
+		const FMeshElement& SelectedMeshElement = SelectedMeshElements[ SelectedElementIndex ];
+		if( SelectedMeshElement.IsValidMeshElement() )
+		{
+			UEditableMesh* EditableMesh = CachedEditableMeshes.FindRef( SelectedMeshElement.ElementAddress.SubMeshAddress );
+			if( EditableMesh != nullptr && EditableMesh->IsCommittedAsInstance() )
+			{
+				return true;
+			}
+		}
+	}
+
+	return false;
+}
+
+const UMeshEditorAssetContainer& FMeshEditorMode::GetAssetContainer() const
+{
+	return *AssetContainer;
+}
+
+void FMeshEditorMode::DeselectMeshElements( const TMap<UEditableMesh*, TArray<FMeshElement>>& MeshElementsToDeselect )
+{
+	FCompoundChangeInput CompoundRevertInput;
+	FSelectOrDeselectMeshElementsChangeInput ChangeInput;
+	for( const auto& MeshAndElements : MeshElementsToDeselect )
+	{
+		for( const auto& MeshElementToDeselect : MeshAndElements.Value )
+		{
+			ChangeInput.MeshElementsToDeselect.Add( MeshElementToDeselect );
+		}
+	}
+
+	CompoundRevertInput.Subchanges.Add( MoveTemp( FSelectOrDeselectMeshElementsChange( MoveTemp( ChangeInput ) ).Execute( MeshEditorModeProxyObject ) ) );
+
+	StoreUndo( MeshEditorModeProxyObject, MakeUnique<FCompoundChange>( MoveTemp( CompoundRevertInput ) ) );
+}
+
+
+bool FMeshEditorMode::DeleteSelectedMeshElement()
+{
+	if( ActiveAction != EMeshEditAction::None )
+	{
+		return false;
+	}
+
+	static TMap<UEditableMesh*, TArray<FMeshElement>> MeshesWithElementsToDelete;
+	GetSelectedMeshesAndElements( EEditableMeshElementType::Any, MeshesWithElementsToDelete );
+	if( MeshesWithElementsToDelete.Num() == 0 )
+	{
+		// @todo should this count as a failure case?
+		return false;
+	}
+
+	FScopedTransaction Transaction( LOCTEXT( "UndoDeleteMeshElement", "Delete" ) );
+
+	CommitSelectedMeshes();
+
+	// Refresh selection (committing may have created a new mesh instance)
+	GetSelectedMeshesAndElements( EEditableMeshElementType::Any, MeshesWithElementsToDelete );
+
+	// Deselect the mesh elements before we delete them.  This will make sure they become selected again after undo.
+	DeselectMeshElements( MeshesWithElementsToDelete );
+
+	for( const auto& MeshAndElements : MeshesWithElementsToDelete )
+	{
+		UEditableMesh* EditableMesh = MeshAndElements.Key;
+
+		EditableMesh->StartModification( EMeshModificationType::Final, EMeshTopologyChange::TopologyChange );
+
+		for( const FMeshElement& MeshElementToDelete : MeshAndElements.Value )
+		{
+			const bool bDeleteOrphanedEdges = true;
+			const bool bDeleteOrphanedVertices = true;
+			const bool bDeleteEmptySections = true;
+
+			// If we deleted the same polygon on multiple selected instances of the same mesh, the polygon could already have been deleted
+			// by the time we get here
+			if( IsElementIDValid( MeshElementToDelete, EditableMesh ) )
+			{
+				if( MeshElementToDelete.ElementAddress.ElementType == EEditableMeshElementType::Vertex )
+				{
+					EditableMesh->DeleteVertexAndConnectedEdgesAndPolygons(
+						FVertexID( MeshElementToDelete.ElementAddress.ElementID ),
+						bDeleteOrphanedEdges,
+						bDeleteOrphanedVertices,
+						bDeleteEmptySections );
+
+				}
+				else if( MeshElementToDelete.ElementAddress.ElementType == EEditableMeshElementType::Edge )
+				{
+					EditableMesh->DeleteEdgeAndConnectedPolygons(
+						FEdgeID( MeshElementToDelete.ElementAddress.ElementID ),
+						bDeleteOrphanedEdges,
+						bDeleteOrphanedVertices,
+						bDeleteEmptySections );
+				}
+				else if( MeshElementToDelete.ElementAddress.ElementType == EEditableMeshElementType::Polygon )
+				{
+					static TArray<FPolygonRef> PolygonRefsToDelete;
+					PolygonRefsToDelete.Reset();
+					PolygonRefsToDelete.Add( FPolygonRef( MeshElementToDelete.ElementAddress.SectionID, FPolygonID( MeshElementToDelete.ElementAddress.ElementID ) ) );
+					EditableMesh->DeletePolygons( PolygonRefsToDelete, bDeleteOrphanedEdges, bDeleteOrphanedVertices, bDeleteEmptySections );
+				}
+			}
+		}
+
+		EditableMesh->EndModification();
+
+		StoreUndo( EditableMesh, EditableMesh->MakeUndo() );
+	}
+
+	// Make sure we're not still hovering over a polygon we deleted
+	StoreUndo( MeshEditorModeProxyObject, MoveTemp( ClearInvalidSelectedElements_Internal() ) );
+
+	return true;
+}
+
+
+bool FMeshEditorMode::TessellateSelectedPolygons()
+{
+	if( ActiveAction != EMeshEditAction::None )
+	{
+		return false;
+	}
+
+	static TMap<UEditableMesh*, TArray<FMeshElement>> SelectedMeshesAndPolygons;
+	GetSelectedMeshesAndPolygons( SelectedMeshesAndPolygons );
+	if( SelectedMeshesAndPolygons.Num() == 0 )
+	{
+		return false;
+	}
+
+	FScopedTransaction Transaction( LOCTEXT( "UndoTessellatePolygon", "Tessellate Polygon" ) );
+
+	CommitSelectedMeshes();
+
+	// Refresh selection (committing may have created a new mesh instance)
+	GetSelectedMeshesAndPolygons( SelectedMeshesAndPolygons );
+
+	// Deselect the mesh elements before we delete them.  This will make sure they become selected again after undo.
+	DeselectMeshElements( SelectedMeshesAndPolygons );
+
+	static TArray<FMeshElement> MeshElementsToSelect;
+	MeshElementsToSelect.Reset();
+
+	for( const auto& SelectedMeshAndPolygons : SelectedMeshesAndPolygons )
+	{
+		UEditableMesh* EditableMesh = SelectedMeshAndPolygons.Key;
+		const TArray<FMeshElement>& PolygonElements = SelectedMeshAndPolygons.Value;
+
+		static TArray<FPolygonRef> PolygonsToTessellate;
+		PolygonsToTessellate.Reset( PolygonElements.Num() );
+
+		for( const auto& PolygonElement : PolygonElements )
+		{
+			const FPolygonRef PolygonRef( PolygonElement.ElementAddress.SectionID, FPolygonID( PolygonElement.ElementAddress.ElementID ) );
+			PolygonsToTessellate.Add( PolygonRef );
+		}
+
+
+		EditableMesh->StartModification( EMeshModificationType::Final, EMeshTopologyChange::TopologyChange );
+
+		// @todo mesheditor: Expose as configurable parameter
+		const ETriangleTessellationMode TriangleTessellationMode = ETriangleTessellationMode::FourTriangles;
+
+		static TArray<FPolygonRef> NewPolygonRefs;
+		EditableMesh->TessellatePolygons( PolygonsToTessellate, TriangleTessellationMode, /* Out */ NewPolygonRefs );
+
+		// Select the new polygons
+		for( const FPolygonRef NewPolygonRef : NewPolygonRefs )
+		{
+			FMeshElement NewPolygonMeshElement;
+			{
+				NewPolygonMeshElement.Component = PolygonElements[0].Component;
+				NewPolygonMeshElement.ElementAddress = PolygonElements[0].ElementAddress;
+				NewPolygonMeshElement.ElementAddress.ElementType = EEditableMeshElementType::Polygon;
+				NewPolygonMeshElement.ElementAddress.SectionID = NewPolygonRef.SectionID;
+				NewPolygonMeshElement.ElementAddress.ElementID = NewPolygonRef.PolygonID;
+			}
+
+			MeshElementsToSelect.Add( NewPolygonMeshElement );
+		}
+
+		EditableMesh->EndModification();
+
+		StoreUndo( EditableMesh, EditableMesh->MakeUndo() );
+	}
+
+	{
+		FCompoundChangeInput CompoundRevertInput;
+
+		// Make sure we're not still hovering over a polygon we removed
+		CompoundRevertInput.Subchanges.Add( MoveTemp( ClearInvalidSelectedElements_Internal() ) );
+
+		// Select the polygon leftover after removing the edge
+		FSelectOrDeselectMeshElementsChangeInput ChangeInput;
+		ChangeInput.MeshElementsToSelect = MeshElementsToSelect;
+		CompoundRevertInput.Subchanges.Add( MoveTemp( FSelectOrDeselectMeshElementsChange( MoveTemp( ChangeInput ) ).Execute( MeshEditorModeProxyObject ) ) );
+
+		StoreUndo( MeshEditorModeProxyObject, MakeUnique<FCompoundChange>( MoveTemp( CompoundRevertInput ) ) );
+	}
+
+	StoreUndo( MeshEditorModeProxyObject, MoveTemp( ClearInvalidSelectedElements_Internal() ) );
+
+	return true;
+}
+
+
+void FMeshEditorMode::AddOrRemoveSubdivisionLevel( const bool bShouldAdd )
+{
+	if( ActiveAction == EMeshEditAction::None )
+	{
+		static TSet< UEditableMesh* > SelectedMeshes;
+		SelectedMeshes.Reset();
+
+		if( SelectedMeshElements.Num() == 0 )
+		{
+			return;
+		}
+
+		FScopedTransaction Transaction( bShouldAdd ?
+											 LOCTEXT( "UndoAddSubdivisionLevel", "Add Subdivision Level" ) :
+											 LOCTEXT( "UndoRemoveSubdivisionLevel", "Remove Subdivision Level" ) );
+
+		CommitSelectedMeshes();
+
+		for( int32 SelectedElementIndex = 0; SelectedElementIndex < SelectedMeshElements.Num(); ++SelectedElementIndex )
+		{
+			const FMeshElement& SelectedMeshElement = SelectedMeshElements[ SelectedElementIndex ];
+			if( SelectedMeshElement.IsValidMeshElement() )
+			{
+				UEditableMesh* EditableMesh = FindOrCreateEditableMesh( *SelectedMeshElement.Component, SelectedMeshElement.ElementAddress.SubMeshAddress );
+				if( EditableMesh != nullptr )
+				{
+					SelectedMeshes.Add( EditableMesh );
+				}
+			}
+		}
+
+		for( UEditableMesh* EditableMesh : SelectedMeshes )
+		{
+			EditableMesh->StartModification( EMeshModificationType::Final, EMeshTopologyChange::TopologyChange );
+
+			if( GIsDemoMode )
+			{
+				// @todo mesheditor demo: Use specific subD count for demo
+				EditableMesh->SetSubdivisionCount( bShouldAdd ? 3 : 0 );
+			}
+			else
+			{
+				EditableMesh->SetSubdivisionCount( FMath::Max( 0, EditableMesh->GetSubdivisionCount() + ( bShouldAdd ? 1 : -1 ) ) );
+			}
+
+			EditableMesh->EndModification();
+
+			StoreUndo( EditableMesh, EditableMesh->MakeUndo() );
+		}
+	}
+}
+
+
+void FMeshEditorMode::QuadrangulateMesh()
+{
+	if( ActiveAction == EMeshEditAction::None )
+	{
+		static TSet< UEditableMesh* > SelectedMeshes;
+		SelectedMeshes.Reset();
+
+		if( SelectedMeshElements.Num() == 0 )
+		{
+			return;
+		}
+
+		FScopedTransaction Transaction( LOCTEXT( "UndoQuadrangulateMesh", "Quadrangulate Mesh" ) );
+
+		CommitSelectedMeshes();
+
+		for( int32 SelectedElementIndex = 0; SelectedElementIndex < SelectedMeshElements.Num(); ++SelectedElementIndex )
+		{
+			const FMeshElement& SelectedMeshElement = SelectedMeshElements[ SelectedElementIndex ];
+			if( SelectedMeshElement.IsValidMeshElement() )
+			{
+				UEditableMesh* EditableMesh = FindOrCreateEditableMesh( *SelectedMeshElement.Component, SelectedMeshElement.ElementAddress.SubMeshAddress );
+				if( EditableMesh != nullptr )
+				{
+					SelectedMeshes.Add( EditableMesh );
+				}
+			}
+		}
+
+		FCompoundChangeInput CompoundRevertInput;
+		CompoundRevertInput.Subchanges.Add( MoveTemp( FDeselectAllMeshElementsChange( FDeselectAllMeshElementsChangeInput() ).Execute( MeshEditorModeProxyObject ) ) );
+		StoreUndo( MeshEditorModeProxyObject, MakeUnique<FCompoundChange>( MoveTemp( CompoundRevertInput ) ) );
+
+		for( UEditableMesh* EditableMesh : SelectedMeshes )
+		{
+			static TArray<FPolygonRef> NewPolygonRefs;
+			EditableMesh->StartModification( EMeshModificationType::Final, EMeshTopologyChange::TopologyChange );
+			EditableMesh->QuadrangulateMesh( NewPolygonRefs );
+			EditableMesh->EndModification();
+
+			StoreUndo( EditableMesh, EditableMesh->MakeUndo() );
+		}
+	}
+}
+
+
+void FMeshEditorMode::FrameSelectedElements( FEditorViewportClient* ViewportClient )
+{
+	FBox BoundingBox( ForceInitToZero );
+
+	switch( GetSelectedMeshElementType() )
+	{
+		case EEditableMeshElementType::Vertex:
+		{
+			static TMap<UEditableMesh*, TArray<TTuple<FMeshElement, FVertexID>>> SelectedMeshesAndVertices;
+			GetSelectedMeshesAndVertices( SelectedMeshesAndVertices );
+			for( const auto& SelectedMeshAndVertices : SelectedMeshesAndVertices )
+			{
+				UEditableMesh* EditableMesh = SelectedMeshAndVertices.Key;
+				const TArray<TTuple<FMeshElement, FVertexID>>& VertexElements = SelectedMeshAndVertices.Value;
+
+				for( const auto& VertexElement : VertexElements )
+				{
+					const UPrimitiveComponent* Component = VertexElement.Get<0>().Component.Get();
+					if( Component )
+					{
+						const FVector VertexPosition = EditableMesh->GetVertexAttribute( VertexElement.Get<1>(), UEditableMeshAttribute::VertexPosition(), 0 );
+						BoundingBox += Component->GetComponentTransform().TransformPosition( VertexPosition );
+					}
+				}
+			}
+			break;
+		}
+
+		case EEditableMeshElementType::Edge:
+		{
+			static TMap<UEditableMesh*, TArray<TTuple<FMeshElement, FEdgeID>>> SelectedMeshesAndEdges;
+			GetSelectedMeshesAndEdges( SelectedMeshesAndEdges );
+			for( const auto& SelectedMeshAndEdges : SelectedMeshesAndEdges )
+			{
+				UEditableMesh* EditableMesh = SelectedMeshAndEdges.Key;
+				const TArray<TTuple<FMeshElement, FEdgeID>>& EdgeElements = SelectedMeshAndEdges.Value;
+
+				for( const auto& EdgeElement : EdgeElements )
+				{
+					const UPrimitiveComponent* Component = EdgeElement.Get<0>().Component.Get();
+					if( Component )
+					{
+						const FVertexID VertexID0 = EditableMesh->GetEdgeVertex( EdgeElement.Get<1>(), 0 );
+						const FVector VertexPosition0 = EditableMesh->GetVertexAttribute( VertexID0, UEditableMeshAttribute::VertexPosition(), 0 );
+						BoundingBox += Component->GetComponentTransform().TransformPosition( VertexPosition0 );
+
+						const FVertexID VertexID1 = EditableMesh->GetEdgeVertex( EdgeElement.Get<1>(), 1 );
+						const FVector VertexPosition1 = EditableMesh->GetVertexAttribute( VertexID1, UEditableMeshAttribute::VertexPosition(), 0 );
+						BoundingBox += Component->GetComponentTransform().TransformPosition( VertexPosition1 );
+					}
+				}
+			}
+			break;
+		}
+
+		case EEditableMeshElementType::Polygon:
+		{
+			static TMap<UEditableMesh*, TArray<TTuple<FMeshElement, FPolygonRef>>> SelectedMeshesAndPolygons;
+			GetSelectedMeshesAndPolygons( SelectedMeshesAndPolygons );
+			for( const auto& SelectedMeshAndPolygons : SelectedMeshesAndPolygons )
+			{
+				UEditableMesh* EditableMesh = SelectedMeshAndPolygons.Key;
+				const TArray<TTuple<FMeshElement, FPolygonRef>>& PolygonElements = SelectedMeshAndPolygons.Value;
+
+				for( const auto& PolygonElement : PolygonElements )
+				{
+					const UPrimitiveComponent* Component = PolygonElement.Get<0>().Component.Get();
+					if( Component )
+					{
+						const int32 PolygonVertexCount = EditableMesh->GetPolygonPerimeterVertexCount( PolygonElement.Get<1>() );
+						for( int32 PolygonVertexIndex = 0; PolygonVertexIndex < PolygonVertexCount; ++PolygonVertexIndex )
+						{
+							const FVector VertexPosition = EditableMesh->GetPolygonPerimeterVertexAttribute( PolygonElement.Get<1>(), PolygonVertexIndex, UEditableMeshAttribute::VertexPosition(), 0 );
+							BoundingBox += Component->GetComponentTransform().TransformPosition( VertexPosition );
+						}
+					}
+				}
+			}
+			break;
+		}
+
+		default:
+			return;
+	}
+
+	ViewportClient->FocusViewportOnBox( BoundingBox );
+}
+
+
+bool FMeshEditorMode::RemoveSelectedEdges()
+{
+	if( ActiveAction != EMeshEditAction::None )
+	{
+		return false;
+	}
+
+	// @todo mesheditor: Only if one edge is selected, for now.  It gets a bit confusing when performing this operation
+	// with more than one edge selected.  However, it can be very useful when collapsing away edges that don't share any
+	// common polygons, so we should try to support it.
+	if( SelectedMeshElements.Num() != 1 )
+	{
+		// @todo should this count as a failure case?
+		return false;
+	}
+
+	static TMap< UEditableMesh*, TArray< FMeshElement > > MeshesWithEdgesToRemove;
+	GetSelectedMeshesAndEdges( MeshesWithEdgesToRemove );
+
+	if( MeshesWithEdgesToRemove.Num() == 0 )
+	{
+		// @todo should this count as a failure case?
+		return false;
+	}
+
+	FScopedTransaction Transaction( LOCTEXT( "UndoRemoveEdge", "Remove Edge" ) );
+
+	CommitSelectedMeshes();
+
+	// Refresh selection (committing may have created a new mesh instance)
+	GetSelectedMeshesAndEdges( MeshesWithEdgesToRemove );
+
+	TArray<FMeshElement> MeshElementsToSelect;
+	for( const auto& MeshAndElements : MeshesWithEdgesToRemove )
+	{
+		UEditableMesh* EditableMesh = MeshAndElements.Key;
+
+		EditableMesh->StartModification( EMeshModificationType::Final, EMeshTopologyChange::TopologyChange );
+
+		const TArray<FMeshElement>& EdgeElementsToRemove = MeshAndElements.Value;
+
+		for( const FMeshElement& EdgeElementToRemove : EdgeElementsToRemove )
+		{
+			const FEdgeID EdgeID( EdgeElementToRemove.ElementAddress.ElementID );
+			{
+				bool bWasEdgeRemoved = false;
+				FPolygonRef NewPolygonRef;
+
+				EditableMesh->TryToRemovePolygonEdge( EdgeID, /* Out */ bWasEdgeRemoved, /* Out */ NewPolygonRef );
+
+				if( bWasEdgeRemoved )
+				{
+					// Select the new polygon
+					FMeshElement NewPolygonMeshElement;
+					{
+						NewPolygonMeshElement.Component = EdgeElementToRemove.Component;
+						NewPolygonMeshElement.ElementAddress = EdgeElementToRemove.ElementAddress;
+						NewPolygonMeshElement.ElementAddress.ElementType = EEditableMeshElementType::Polygon;
+						NewPolygonMeshElement.ElementAddress.SectionID = NewPolygonRef.SectionID;
+						NewPolygonMeshElement.ElementAddress.ElementID = NewPolygonRef.PolygonID;
+					}
+
+					MeshElementsToSelect.Add( NewPolygonMeshElement );
+				}
+				else
+				{
+					// Couldn't remove the edge
+					// @todo mesheditor: Needs good user feedback when this happens
+				}
+			}
+		}
+
+		EditableMesh->EndModification();
+
+		StoreUndo( EditableMesh, EditableMesh->MakeUndo() );
+	}
+
+	{
+		FCompoundChangeInput CompoundRevertInput;
+
+		// Make sure we're not still hovering over a polygon we removed
+		CompoundRevertInput.Subchanges.Add( MoveTemp( ClearInvalidSelectedElements_Internal() ) );
+
+		// Select the polygon leftover after removing the edge
+		FSelectOrDeselectMeshElementsChangeInput ChangeInput;
+		ChangeInput.MeshElementsToSelect = MeshElementsToSelect;
+		CompoundRevertInput.Subchanges.Add( MoveTemp( FSelectOrDeselectMeshElementsChange( MoveTemp( ChangeInput ) ).Execute( MeshEditorModeProxyObject ) ) );
+
+		StoreUndo( MeshEditorModeProxyObject, MakeUnique<FCompoundChange>( MoveTemp( CompoundRevertInput ) ) );
+	}
+
+	return true;
+}
+
+
+bool FMeshEditorMode::SelectEdgeLoops()
+{
+	if( ActiveAction != EMeshEditAction::None )
+	{
+		return false;
+	}
+
+	static TMap< UEditableMesh*, TArray< FMeshElement > > MeshesWithEdgesToRemove;
+	GetSelectedMeshesAndEdges( MeshesWithEdgesToRemove );
+
+	if( MeshesWithEdgesToRemove.Num() == 0 )
+	{
+		// @todo should this count as a failure case?
+		return false;
+	}
+
+	FScopedTransaction Transaction( LOCTEXT( "SelectEdgeLoops", "Select Edge Loops" ) );
+
+	TArray<FMeshElement> MeshElementsToSelect;
+
+	for( const auto& MeshAndElements : MeshesWithEdgesToRemove )
+	{
+		UEditableMesh* EditableMesh = MeshAndElements.Key;
+
+		const TArray<FMeshElement>& SelectedEdgeElements = MeshAndElements.Value;
+		
+		TArray<FEdgeID> UniqueEdgeIDsPerMesh;
+
+		for( const FMeshElement& SelectedEdgeElement : SelectedEdgeElements )
+		{
+			const FEdgeID EdgeID( SelectedEdgeElement.ElementAddress.ElementID );
+			TArray<FEdgeID> EdgeLoopIDs;
+			EditableMesh->GetEdgeLoopElements( EdgeID, EdgeLoopIDs );
+			for( const FEdgeID EdgeLoopID : EdgeLoopIDs )
+			{
+				UniqueEdgeIDsPerMesh.AddUnique( EdgeLoopID );
+			}
+		}
+
+		for( const FEdgeID UniqueEdgeID : UniqueEdgeIDsPerMesh )
+		{
+			MeshElementsToSelect.Emplace( SelectedEdgeElements[ 0 ].Component.Get(), EditableMesh->GetSubMeshAddress(), UniqueEdgeID );
+		}
+	}
+
+	{
+		FCompoundChangeInput CompoundRevertInput;
+		CompoundRevertInput.Subchanges.Add( MoveTemp( FDeselectAllMeshElementsChange( FDeselectAllMeshElementsChangeInput() ).Execute( MeshEditorModeProxyObject ) ) );
+
+		FSelectOrDeselectMeshElementsChangeInput ChangeInput;
+		ChangeInput.MeshElementsToSelect = MeshElementsToSelect;
+		CompoundRevertInput.Subchanges.Add( MoveTemp( FSelectOrDeselectMeshElementsChange( MoveTemp( ChangeInput ) ).Execute( MeshEditorModeProxyObject ) ) );
+
+		StoreUndo( MeshEditorModeProxyObject, MakeUnique<FCompoundChange>( MoveTemp( CompoundRevertInput ) ) );
+	}
+
+	return true;
+}
+
+
+bool FMeshEditorMode::RemoveSelectedVertices()
+{
+	if( ActiveAction != EMeshEditAction::None )
+	{
+		return false;
+	}
+
+	// @todo mesheditor: Only if one vertex is selected, for now.  It gets a bit confusing when performing this operation
+	// with more than one vertex selected.  However, it can be very useful when collapsing away vertices that don't share any
+	// common polygons, so we should try to support it.
+	if( SelectedMeshElements.Num() != 1 )
+	{
+		// @todo should this count as a failure case?
+		return false;
+	}
+
+	static TMap< UEditableMesh*, TArray<FMeshElement> > MeshesWithVerticesToRemove;
+	GetSelectedMeshesAndVertices( MeshesWithVerticesToRemove );
+
+	// Remove selected vertices
+	if( MeshesWithVerticesToRemove.Num() == 0 )
+	{
+		// @todo should this count as a failure case?
+		return false;
+	}
+
+	FScopedTransaction Transaction( LOCTEXT( "UndoRemoveVertex", "Remove Vertex" ) );
+
+	CommitSelectedMeshes();
+
+	// Refresh selection (committing may have created a new mesh instance)
+	GetSelectedMeshesAndVertices( MeshesWithVerticesToRemove );
+
+	// Deselect the mesh elements before we delete them.  This will make sure they become selected again after undo.
+	DeselectMeshElements( MeshesWithVerticesToRemove );
+
+	TArray<FMeshElement> MeshElementsToSelect;
+	for( const auto& MeshAndElements : MeshesWithVerticesToRemove )
+	{
+		UEditableMesh* EditableMesh = MeshAndElements.Key;
+
+		EditableMesh->StartModification( EMeshModificationType::Final, EMeshTopologyChange::TopologyChange );
+
+		const TArray<FMeshElement>& VertexElementsToRemove = MeshAndElements.Value;
+
+		for( const FMeshElement& VertexElementToRemove : VertexElementsToRemove )
+		{
+			const FVertexID VertexID( VertexElementToRemove.ElementAddress.ElementID );
+			{
+				bool bWasVertexRemoved = false;
+				FEdgeID NewEdgeID = FEdgeID::Invalid;
+
+				EditableMesh->TryToRemoveVertex( VertexID, /* Out */ bWasVertexRemoved, /* Out */ NewEdgeID );
+
+				if( bWasVertexRemoved )
+				{
+					// Select the new edge
+					FMeshElement NewEdgeMeshElement;
+					{
+						NewEdgeMeshElement.Component = VertexElementToRemove.Component;
+						NewEdgeMeshElement.ElementAddress = VertexElementToRemove.ElementAddress;
+						NewEdgeMeshElement.ElementAddress.ElementType = EEditableMeshElementType::Edge;
+						NewEdgeMeshElement.ElementAddress.ElementID = NewEdgeID;
+					}
+
+					MeshElementsToSelect.Add( NewEdgeMeshElement );
+				}
+				else
+				{
+					// Couldn't remove the vertex
+					// @todo mesheditor: Needs good user feedback when this happens
+					// @todo mesheditor: If this fails, it will already potentially have created a new instance. To be 100% correct, it needs to do a prepass
+					// to determine whether the operation can complete successfully before actually doing it.
+				}
+			}
+		}
+
+		EditableMesh->EndModification();
+
+		StoreUndo( EditableMesh, EditableMesh->MakeUndo() );
+	}
+
+	{
+		FCompoundChangeInput CompoundRevertInput;
+
+		// Make sure we're not still hovering over a polygon we removed
+		CompoundRevertInput.Subchanges.Add( MoveTemp( ClearInvalidSelectedElements_Internal() ) );
+
+		// Select the edge that the vertex we removed was a member of
+		FSelectOrDeselectMeshElementsChangeInput ChangeInput;
+		ChangeInput.MeshElementsToSelect = MeshElementsToSelect;
+		CompoundRevertInput.Subchanges.Add( MoveTemp( FSelectOrDeselectMeshElementsChange( MoveTemp( ChangeInput ) ).Execute( MeshEditorModeProxyObject ) ) );
+
+		StoreUndo( MeshEditorModeProxyObject, MakeUnique<FCompoundChange>( MoveTemp( CompoundRevertInput ) ) );
+	}
+
+	return true;
+}
+
+
+bool FMeshEditorMode::WeldSelectedVertices()
+{
+	if( ActiveAction != EMeshEditAction::None )
+	{
+		return false;
+	}
+
+	static TMap< UEditableMesh*, TArray<FMeshElement> > MeshesWithVerticesToWeld;
+	GetSelectedMeshesAndVertices( MeshesWithVerticesToWeld );
+
+	if( MeshesWithVerticesToWeld.Num() == 0 )
+	{
+		// @todo should this count as a failure case?
+		return false;
+	}
+
+	const FScopedTransaction Transaction( LOCTEXT( "UndoWeldVertices", "Weld Vertices" ) );
+
+	CommitSelectedMeshes();
+
+	// Refresh selection (committing may have created a new mesh instance)
+	GetSelectedMeshesAndVertices( MeshesWithVerticesToWeld );
+
+	// Deselect the mesh elements before we delete them.  This will make sure they become selected again after undo.
+	DeselectMeshElements( MeshesWithVerticesToWeld );
+
+	TArray<FMeshElement> MeshElementsToSelect;
+	for( auto& MeshAndElements : MeshesWithVerticesToWeld )
+	{
+		UEditableMesh* EditableMesh = MeshAndElements.Key;
+		const TArray<FMeshElement>& VertexElementsToWeld = MeshAndElements.Value;
+
+		if( VertexElementsToWeld.Num() < 2 )
+		{
+			continue;
+		}
+
+		static TArray<FVertexID> VertexIDsToWeld;
+		VertexIDsToWeld.Reset( MeshAndElements.Value.Num() );
+
+		EditableMesh->StartModification( EMeshModificationType::Final, EMeshTopologyChange::TopologyChange );
+
+		for( const FMeshElement& VertexElementToWeld : VertexElementsToWeld )
+		{
+			const FVertexID VertexID( VertexElementToWeld.ElementAddress.ElementID );
+			VertexIDsToWeld.Add( VertexID );
+		}
+
+		FVertexID WeldedVertexID = FVertexID::Invalid;
+		EditableMesh->WeldVertices( VertexIDsToWeld, WeldedVertexID );
+
+		if( WeldedVertexID != FVertexID::Invalid )
+		{
+			FMeshElement NewVertexMeshElement;
+			{
+				NewVertexMeshElement.Component = VertexElementsToWeld[ 0 ].Component;
+				NewVertexMeshElement.ElementAddress = VertexElementsToWeld[ 0 ].ElementAddress;
+				NewVertexMeshElement.ElementAddress.ElementType = EEditableMeshElementType::Vertex;
+				NewVertexMeshElement.ElementAddress.ElementID = WeldedVertexID;
+			}
+
+			MeshElementsToSelect.Add( NewVertexMeshElement );
+		}
+		else
+		{
+			// Couldn't weld the vertices
+			// @todo mesheditor: Needs good user feedback when this happens
+			// @todo mesheditor: If this fails, it will already potentially have created a new instance. To be 100% correct, it needs to do a prepass
+			// to determine whether the operation can complete successfully before actually doing it.
+		}
+
+		EditableMesh->EndModification();
+
+		StoreUndo( EditableMesh, EditableMesh->MakeUndo() );
+	}
+
+	{
+		FCompoundChangeInput CompoundRevertInput;
+
+		CompoundRevertInput.Subchanges.Add( MoveTemp( ClearInvalidSelectedElements_Internal() ) );
+
+		// Select the vertex that was welded.
+		FSelectOrDeselectMeshElementsChangeInput ChangeInput;
+		ChangeInput.MeshElementsToSelect = MeshElementsToSelect;
+		CompoundRevertInput.Subchanges.Add( MoveTemp( FSelectOrDeselectMeshElementsChange( MoveTemp( ChangeInput ) ).Execute( MeshEditorModeProxyObject ) ) );
+
+		StoreUndo( MeshEditorModeProxyObject, MakeUnique<FCompoundChange>( MoveTemp( CompoundRevertInput ) ) );
+	}
+
+	return true;
+}
+
+
+bool FMeshEditorMode::FlipSelectedPolygons()
+{
+	if( ActiveAction != EMeshEditAction::None )
+	{
+		return false;
+	}
+
+	static TMap< UEditableMesh*, TArray<FMeshElement> > MeshesAndPolygons;
+	GetSelectedMeshesAndPolygons( /* Out */ MeshesAndPolygons );
+
+	if( MeshesAndPolygons.Num() == 0 )
+	{
+		// @todo should this count as a failure case?
+		return false;
+	}
+
+	const FScopedTransaction Transaction( LOCTEXT( "UndoFlipPolygon", "Flip Polygon" ) );
+
+	CommitSelectedMeshes();
+
+	// Refresh selection (committing may have created a new mesh instance)
+	GetSelectedMeshesAndPolygons( /* Out */ MeshesAndPolygons );
+
+	// Flip selected polygons
+	for( const auto& MeshAndPolygons : MeshesAndPolygons )
+	{
+		UEditableMesh* EditableMesh = MeshAndPolygons.Key;
+
+		EditableMesh->StartModification( EMeshModificationType::Final, EMeshTopologyChange::TopologyChange );
+
+		static TArray<FPolygonRef> PolygonsToFlip;
+		PolygonsToFlip.Reset();
+
+		for( const FMeshElement& PolygonElement : MeshAndPolygons.Value )
+		{
+			const FPolygonRef PolygonRef( PolygonElement.ElementAddress.SectionID, FPolygonID( PolygonElement.ElementAddress.ElementID ) );
+			PolygonsToFlip.Add( PolygonRef );
+		}
+
+		EditableMesh->FlipPolygons( PolygonsToFlip );
+
+		EditableMesh->EndModification();
+
+		StoreUndo( EditableMesh, EditableMesh->MakeUndo() );
+	}
+
+	return true;
+}
+
+
+bool FMeshEditorMode::TriangulateSelectedPolygons()
+{
+	if( ActiveAction != EMeshEditAction::None )
+	{
+		return false;
+	}
+
+	static TMap< UEditableMesh*, TArray<FMeshElement> > MeshesAndPolygons;
+	GetSelectedMeshesAndPolygons( /* Out */ MeshesAndPolygons );
+
+	if( MeshesAndPolygons.Num() == 0 )
+	{
+		// @todo should this count as a failure case?
+		return false;
+	}
+
+	const FScopedTransaction Transaction( LOCTEXT( "UndoTrianglulatePolygon", "Triangulate Polygon" ) );
+
+	CommitSelectedMeshes();
+
+	// Refresh selection (committing may have created a new mesh instance)
+	GetSelectedMeshesAndPolygons( /* Out */ MeshesAndPolygons );
+
+	// Deselect the polygons first.  They'll be deleted and replaced by triangles.  This will also make sure 
+	// they become selected again after undo.
+	DeselectMeshElements( MeshesAndPolygons );
+
+	static TArray<FMeshElement> MeshElementsToSelect;
+	MeshElementsToSelect.Reset();
+
+	for( const auto& MeshAndPolygons : MeshesAndPolygons )
+	{
+		UEditableMesh* EditableMesh = MeshAndPolygons.Key;
+		
+		UPrimitiveComponent* Component = nullptr;
+		for( const FMeshElement& PolygonElement : MeshAndPolygons.Value )
+		{
+			Component = PolygonElement.Component.Get();
+			break;
+		}
+		check( Component != nullptr );
+
+
+		EditableMesh->StartModification( EMeshModificationType::Final, EMeshTopologyChange::TopologyChange );
+
+		static TArray<FPolygonRef> PolygonsToTriangulate;
+		PolygonsToTriangulate.Reset( MeshAndPolygons.Value.Num() );
+
+		for( const FMeshElement& PolygonElement : MeshAndPolygons.Value )
+		{
+			const FPolygonRef PolygonRef( PolygonElement.ElementAddress.SectionID, FPolygonID( PolygonElement.ElementAddress.ElementID ) );
+			PolygonsToTriangulate.Add( PolygonRef );
+		}
+
+		static TArray<FPolygonRef> NewTrianglePolygonRefs;
+		EditableMesh->TriangulatePolygons( PolygonsToTriangulate, /* Out */ NewTrianglePolygonRefs );
+
+		for( const FPolygonRef& NewTrianglePolygonRef : NewTrianglePolygonRefs )
+		{
+			// Select the new polygon
+			FMeshElement NewPolygonMeshElement;
+			{
+				NewPolygonMeshElement.Component = Component;
+				NewPolygonMeshElement.ElementAddress.SubMeshAddress = EditableMesh->GetSubMeshAddress();
+				NewPolygonMeshElement.ElementAddress.ElementType = EEditableMeshElementType::Polygon;
+				NewPolygonMeshElement.ElementAddress.SectionID = NewTrianglePolygonRef.SectionID;
+				NewPolygonMeshElement.ElementAddress.ElementID = NewTrianglePolygonRef.PolygonID;
+			}
+
+			MeshElementsToSelect.Add( NewPolygonMeshElement );
+		}
+
+		EditableMesh->EndModification();
+
+		StoreUndo( EditableMesh, EditableMesh->MakeUndo() );
+	}
+
+	// Select the newly-created triangles
+	{
+		FCompoundChangeInput CompoundRevertInput;
+
+		// Make sure we're not still hovering over a polygon we removed
+		CompoundRevertInput.Subchanges.Add( MoveTemp( ClearInvalidSelectedElements_Internal() ) );
+
+		FSelectOrDeselectMeshElementsChangeInput ChangeInput;
+		ChangeInput.MeshElementsToSelect = MeshElementsToSelect;
+		CompoundRevertInput.Subchanges.Add( MoveTemp( FSelectOrDeselectMeshElementsChange( MoveTemp( ChangeInput ) ).Execute( MeshEditorModeProxyObject ) ) );
+
+		StoreUndo( MeshEditorModeProxyObject, MakeUnique<FCompoundChange>( MoveTemp( CompoundRevertInput ) ) );
+	}
+
+
+	return true;
+}
+
+
+bool FMeshEditorMode::AssignSelectedMaterialToSelectedPolygons()
+{
+	IContentBrowserSingleton& ContentBrowser = FModuleManager::LoadModuleChecked<FContentBrowserModule>( "ContentBrowser" ).Get();
+
+	static TArray<FAssetData> SelectedAssets;
+	ContentBrowser.GetSelectedAssets( SelectedAssets );
+
+	UMaterialInterface* SelectedMaterial = FAssetData::GetFirstAsset<UMaterialInterface>( SelectedAssets );
+
+	return AssignMaterialToSelectedPolygons( SelectedMaterial );
+}
+
+
+bool FMeshEditorMode::AssignMaterialToSelectedPolygons( UMaterialInterface* SelectedMaterial )
+{
+	if( SelectedMaterial )
+	{
+		if( ActiveAction != EMeshEditAction::None )
+		{
+			return false;
+		}
+
+		static TMap< UEditableMesh*, TArray<FMeshElement> > MeshesAndPolygons;
+		GetSelectedMeshesAndPolygons( /* Out */ MeshesAndPolygons );
+
+		if( MeshesAndPolygons.Num() == 0 )
+		{
+			// @todo should this count as a failure case?
+			return false;
+		}
+
+		static TArray<FPolygonRef> NewPolygonRefs;
+		NewPolygonRefs.Reset();
+
+		static TArray<FMeshElement> MeshElementsToSelect;
+		MeshElementsToSelect.Reset();
+
+		const FScopedTransaction Transaction( LOCTEXT( "UndoAssignMaterialToPolygon", "Assign Material to Polygon" ) );
+
+		CommitSelectedMeshes();
+
+		// Refresh selection (committing may have created a new mesh instance)
+		GetSelectedMeshesAndPolygons( /* Out */ MeshesAndPolygons );
+
+		// Deselect the polygons first.  They'll be moved to a new section, and hence be assigned new PolygonRefs.  This will also make sure 
+		// they become selected again after undo.
+		DeselectMeshElements( MeshesAndPolygons );
+
+		for( const auto& MeshAndPolygons : MeshesAndPolygons )
+		{
+			UEditableMesh* EditableMesh = MeshAndPolygons.Key;
+
+			UPrimitiveComponent* Component = nullptr;
+			for( const FMeshElement& PolygonElement : MeshAndPolygons.Value )
+			{
+				Component = PolygonElement.Component.Get();
+				break;
+			}
+			check( Component != nullptr );
+
+			EditableMesh->StartModification( EMeshModificationType::Final, EMeshTopologyChange::TopologyChange );
+			{
+				static TArray<FPolygonRef> PolygonsToAssign;
+				PolygonsToAssign.Reset();
+
+				for( const FMeshElement& PolygonElement : MeshAndPolygons.Value )
+				{
+					const FPolygonRef PolygonRef( PolygonElement.ElementAddress.SectionID, FPolygonID( PolygonElement.ElementAddress.ElementID ) );
+					PolygonsToAssign.Add( PolygonRef );
+				}
+
+				EditableMesh->AssignMaterialToPolygons( PolygonsToAssign, SelectedMaterial, NewPolygonRefs );
+
+				for( const FPolygonRef& NewPolygonRef : NewPolygonRefs )
+				{
+					// Select the new polygon
+					FMeshElement NewPolygonMeshElement;
+					{
+						NewPolygonMeshElement.Component = Component;
+						NewPolygonMeshElement.ElementAddress.SubMeshAddress = EditableMesh->GetSubMeshAddress();
+						NewPolygonMeshElement.ElementAddress.ElementType = EEditableMeshElementType::Polygon;
+						NewPolygonMeshElement.ElementAddress.SectionID = NewPolygonRef.SectionID;
+						NewPolygonMeshElement.ElementAddress.ElementID = NewPolygonRef.PolygonID;
+					}
+
+					MeshElementsToSelect.Add( NewPolygonMeshElement );
+				}
+			}
+			EditableMesh->EndModification();
+
+			StoreUndo( EditableMesh, EditableMesh->MakeUndo() );
+		}
+
+		// Select the newly-created triangles
+		{
+			FCompoundChangeInput CompoundRevertInput;
+
+			// Make sure we're not still hovering over a polygon we removed
+			CompoundRevertInput.Subchanges.Add( MoveTemp( ClearInvalidSelectedElements_Internal() ) );
+
+			FSelectOrDeselectMeshElementsChangeInput ChangeInput;
+			ChangeInput.MeshElementsToSelect = MeshElementsToSelect;
+			CompoundRevertInput.Subchanges.Add( MoveTemp( FSelectOrDeselectMeshElementsChange( MoveTemp( ChangeInput ) ).Execute( MeshEditorModeProxyObject ) ) );
+
+			StoreUndo( MeshEditorModeProxyObject, MakeUnique<FCompoundChange>( MoveTemp( CompoundRevertInput ) ) );
+		}
+
+		return true;
+	}
+
+	return true;
+}
+
+
+bool FMeshEditorMode::MakeSelectedEdgesHardOrSoft(const bool bMakeEdgesHard)
+{
+	if( ActiveAction != EMeshEditAction::None )
+	{
+		return false;
+	}
+
+	// Make edges hard or soft
+	static TMap< UEditableMesh*, TArray<TTuple<FMeshElement, FEdgeID>> > MeshesAndEdges;
+	GetSelectedMeshesAndEdges( /* Out */ MeshesAndEdges );
+
+	if( MeshesAndEdges.Num() == 0 )
+	{
+		// @todo should this count as a failure case?
+		return false;
+	}
+
+	const FScopedTransaction Transaction( bMakeEdgesHard ? LOCTEXT( "UndoHardenEdge", "Harden Edge" ) : LOCTEXT( "UndoSoftenEdge", "Soften Edge" ) );
+
+	CommitSelectedMeshes();
+
+	// Refresh selection (committing may have created a new mesh instance)
+	GetSelectedMeshesAndEdges( /* Out */ MeshesAndEdges );
+
+	for( auto& MeshAndEdges : MeshesAndEdges )
+	{
+		static TArray<FEdgeID> EdgeIDs;
+		EdgeIDs.Reset();
+
+		static TArray<bool> EdgesNewIsHard;
+		EdgesNewIsHard.Reset();
+
+		for( TTuple<FMeshElement, FEdgeID>& EdgeElement : MeshAndEdges.Value )
+		{
+			EdgeIDs.Add( EdgeElement.Get<1>() );
+			EdgesNewIsHard.Add( bMakeEdgesHard );
+		}
+
+
+		UEditableMesh* EditableMesh = MeshAndEdges.Key;
+		verify( !EditableMesh->AnyChangesToUndo() );
+
+		EditableMesh->StartModification( EMeshModificationType::Final, EMeshTopologyChange::TopologyChange );
+
+		EditableMesh->SetEdgesHardness( EdgeIDs, EdgesNewIsHard );
+
+		EditableMesh->EndModification();
+
+		StoreUndo( EditableMesh, EditableMesh->MakeUndo() );
+	}
+
+	return true;
+}
+
+
+bool FMeshEditorMode::InputAxis( FEditorViewportClient* ViewportClient, FViewport* Viewport, int32 ControllerId, FKey Key, float Delta, float DeltaTime )
+{
+	bool bHandled = false;
+
+	return bHandled ? true : FEdMode::InputAxis( ViewportClient, Viewport, ControllerId, Key, Delta, DeltaTime );
+}
+
+
+bool FMeshEditorMode::InputDelta( FEditorViewportClient* ViewportClient, FViewport* Viewport, FVector& Drag, FRotator& Rotation, FVector& Scale )
+{
+	bool bHandled = false;
+
+	return bHandled ? true : FEdMode::InputDelta( ViewportClient, Viewport, Drag, Rotation, Scale );
+}
+
+
+bool FMeshEditorMode::IsCompatibleWith(FEditorModeID OtherModeID) const
+{
+	// We are compatible with all other modes!
+	return true;
+}
+
+
+void FMeshEditorMode::AddReferencedObjects( FReferenceCollector& Collector )
+{
+	Collector.AddReferencedObject( ActiveActionInteractor );
+
+	Collector.AddReferencedObject( HoveredGeometryMaterial );
+	Collector.AddReferencedObject( HoveredFaceMaterial );
+
+	for( auto Pair : CachedEditableMeshes )
+	{
+		Collector.AddReferencedObject( Pair.Value );
+	}
+
+	for( TTuple<UObject*, TUniquePtr<FChange>>& ObjectAndPreviewRevertChange : PreviewRevertChanges )
+	{
+		UObject* Object = ObjectAndPreviewRevertChange.Get<0>();
+		Collector.AddReferencedObject( Object );
+	}
+
+	Collector.AddReferencedObject( MeshEditorModeProxyObject );
+	Collector.AddReferencedObject( AssetContainer );
+}
+
+
+void FMeshEditorMode::AddVertexToDynamicMesh( const UEditableMesh& EditableMesh, const FTransform& CameraToWorld, const FMatrix& ComponentToWorldMatrix, const FVertexID VertexID, const FColor ColorAndOpacity, const float SizeBias, const bool bApplyDepthBias, FDynamicMeshBuilder& MeshBuilder )
+{
+	const FVector VertexPosition = ComponentToWorldMatrix.TransformPosition( EditableMesh.GetVertexAttribute( VertexID, UEditableMeshAttribute::VertexPosition(), 0 ) );
+
+	const float DistanceToCamera = ( CameraToWorld.GetLocation() - VertexPosition ).Size();
+	const float DistanceBasedScaling = MeshEd::OverlayDistanceBias->GetFloat() + DistanceToCamera * MeshEd::OverlayDistanceScaleFactor->GetFloat();
+
+	const float SpriteSize = ( MeshEd::OverlayVertexSize->GetFloat() + SizeBias ) * DistanceBasedScaling;
+	const float HalfSpriteSize = SpriteSize * 0.5f;
+
+	const FVector CameraUpVector = CameraToWorld.TransformVector( FVector::UpVector );
+	const FVector CameraRightVector = CameraToWorld.TransformVector( FVector::RightVector );
+	const FVector DirectionToCamera = ( CameraToWorld.GetLocation() - VertexPosition ).GetSafeNormal();
+	const FVector RightDirection = CameraRightVector.GetSafeNormal();
+	const FVector UpDirection = CameraUpVector.GetSafeNormal();
+
+	// We're offsetting the geometry from the actual face a bit, to avoid z-fighting for this particular effect
+	const FVector VertexOffset = bApplyDepthBias ? ( DirectionToCamera * MeshEd::OverlayDepthOffset->GetFloat() * DistanceBasedScaling ) : FVector::ZeroVector;
+
+	FVector QuadPositions[ 4 ];
+	QuadPositions[ 0 ] = ( VertexPosition - RightDirection * HalfSpriteSize - UpDirection * HalfSpriteSize ) + VertexOffset;
+	QuadPositions[ 1 ] = ( VertexPosition - RightDirection * HalfSpriteSize + UpDirection * HalfSpriteSize ) + VertexOffset;
+	QuadPositions[ 2 ] = ( VertexPosition + RightDirection * HalfSpriteSize + UpDirection * HalfSpriteSize ) + VertexOffset;
+	QuadPositions[ 3 ] = ( VertexPosition + RightDirection * HalfSpriteSize - UpDirection * HalfSpriteSize ) + VertexOffset;
+
+	const int32 FirstVertexIndex = 
+		MeshBuilder.AddVertex( VertexOffset + QuadPositions[ 0 ], FVector2D( 0, 0 ), FVector( 1, 0, 0 ), FVector( 0, 1, 0 ), FVector( 0, 0, 1 ), ColorAndOpacity );
+	MeshBuilder.AddVertex( VertexOffset + QuadPositions[ 1 ], FVector2D( 0, 1 ), FVector( 1, 0, 0 ), FVector( 0, 1, 0 ), FVector( 0, 0, 1 ), ColorAndOpacity );
+	MeshBuilder.AddVertex( VertexOffset + QuadPositions[ 2 ], FVector2D( 1, 1 ), FVector( 1, 0, 0 ), FVector( 0, 1, 0 ), FVector( 0, 0, 1 ), ColorAndOpacity );
+	MeshBuilder.AddVertex( VertexOffset + QuadPositions[ 3 ], FVector2D( 1, 0 ), FVector( 1, 0, 0 ), FVector( 0, 1, 0 ), FVector( 0, 0, 1 ), ColorAndOpacity );
+
+	MeshBuilder.AddTriangle( FirstVertexIndex + 0, FirstVertexIndex + 1, FirstVertexIndex + 2 );
+	MeshBuilder.AddTriangle( FirstVertexIndex + 0, FirstVertexIndex + 2, FirstVertexIndex + 3 );
+}
+
+
+void FMeshEditorMode::AddThickLineToDynamicMesh( const FTransform& CameraToWorld, const FVector EdgeVertexPositions[2], const FColor ColorAndOpacity, const float SizeBias, const bool bApplyDepthBias, FDynamicMeshBuilder& MeshBuilder )
+{
+	const float DistanceBias = MeshEd::OverlayDistanceBias->GetFloat();
+	const float Vertex0DistanceToCamera = ( CameraToWorld.GetLocation() - EdgeVertexPositions[ 0 ] ).Size();
+	const float Vertex0DistanceBasedScaling = DistanceBias + Vertex0DistanceToCamera * MeshEd::OverlayDistanceScaleFactor->GetFloat();
+	const float Vertex1DistanceToCamera = ( CameraToWorld.GetLocation() - EdgeVertexPositions[ 1 ] ).Size();
+	const float Vertex1DistanceBasedScaling = DistanceBias + Vertex1DistanceToCamera * MeshEd::OverlayDistanceScaleFactor->GetFloat();
+
+	const float Thickness = ( MeshEd::OverlayLineThickness->GetFloat() + SizeBias );
+	const float HalfThickness = Thickness * 0.5f;
+
+	const FVector DirectionToCamera0 = ( CameraToWorld.GetLocation() - EdgeVertexPositions[ 0 ] ).GetSafeNormal();
+	const FVector DirectionToCamera1 = ( CameraToWorld.GetLocation() - EdgeVertexPositions[ 1 ] ).GetSafeNormal();
+	const FVector EdgeForward = ( EdgeVertexPositions[ 1 ] - EdgeVertexPositions[ 0 ] ).GetSafeNormal();
+	const FVector EdgeRight0 = FVector::CrossProduct( EdgeForward, DirectionToCamera0 ).GetSafeNormal();
+	const FVector EdgeRight1 = FVector::CrossProduct( EdgeForward, DirectionToCamera1 ).GetSafeNormal();
+
+	// We're offsetting the geometry from the actual face a bit, to avoid z-fighting for this particular effect
+	// @todo mesheditor urgent: Can we use a hardware depth bias instead?
+	const FVector Vertex0Offset = bApplyDepthBias ? ( DirectionToCamera0 * MeshEd::OverlayDepthOffset->GetFloat() * Vertex0DistanceBasedScaling ) : FVector::ZeroVector;
+	const FVector Vertex1Offset = bApplyDepthBias ? ( DirectionToCamera1 * MeshEd::OverlayDepthOffset->GetFloat() * Vertex1DistanceBasedScaling ) : FVector::ZeroVector;
+
+	FVector QuadPositions[ 4 ];
+	QuadPositions[ 0 ] = ( EdgeVertexPositions[ 0 ] - EdgeRight0 * HalfThickness * Vertex0DistanceBasedScaling ) + Vertex0Offset;
+	QuadPositions[ 1 ] = ( EdgeVertexPositions[ 0 ] + EdgeRight0 * HalfThickness * Vertex0DistanceBasedScaling ) + Vertex0Offset;
+	QuadPositions[ 2 ] = ( EdgeVertexPositions[ 1 ] + EdgeRight1 * HalfThickness * Vertex1DistanceBasedScaling ) + Vertex1Offset;
+	QuadPositions[ 3 ] = ( EdgeVertexPositions[ 1 ] - EdgeRight1 * HalfThickness * Vertex1DistanceBasedScaling ) + Vertex1Offset;
+
+	const int32 FirstVertexIndex = 
+		MeshBuilder.AddVertex( QuadPositions[ 0 ], FVector2D( 0, 0 ), FVector( 1, 0, 0 ), FVector( 0, 1, 0 ), FVector( 0, 0, 1 ), ColorAndOpacity );
+	MeshBuilder.AddVertex( QuadPositions[ 1 ], FVector2D( 0, 1 ), FVector( 1, 0, 0 ), FVector( 0, 1, 0 ), FVector( 0, 0, 1 ), ColorAndOpacity );
+	MeshBuilder.AddVertex( QuadPositions[ 2 ], FVector2D( 1, 1 ), FVector( 1, 0, 0 ), FVector( 0, 1, 0 ), FVector( 0, 0, 1 ), ColorAndOpacity );
+	MeshBuilder.AddVertex( QuadPositions[ 3 ], FVector2D( 1, 0 ), FVector( 1, 0, 0 ), FVector( 0, 1, 0 ), FVector( 0, 0, 1 ), ColorAndOpacity );
+
+	MeshBuilder.AddTriangle( FirstVertexIndex + 0, FirstVertexIndex + 1, FirstVertexIndex + 2 );
+	MeshBuilder.AddTriangle( FirstVertexIndex + 0, FirstVertexIndex + 2, FirstVertexIndex + 3 );
+}
+
+
+void FMeshEditorMode::AddEdgeToDynamicMesh( const UEditableMesh& EditableMesh, const FTransform& CameraToWorld, const FMatrix& ComponentToWorldMatrix, const FEdgeID EdgeID, const FColor ColorAndOpacity, const float SizeBias, FDynamicMeshBuilder& MeshBuilder )
+{
+	FVertexID MeshVertexIDs[ 2 ];
+	EditableMesh.GetEdgeVertices( EdgeID, /* Out */ MeshVertexIDs[0], /* Out */ MeshVertexIDs[1] );
+
+	static TArray<FPolygonRef> ConnectedPolygons;
+	EditableMesh.GetEdgeConnectedPolygons( EdgeID, /* Out */ ConnectedPolygons );
+
+	FVector EdgeVertexPositions[ 2 ];
+	for( uint32 EdgeVertexNumber = 0; EdgeVertexNumber < ARRAY_COUNT( MeshVertexIDs ); ++EdgeVertexNumber )
+	{
+		EdgeVertexPositions[ EdgeVertexNumber ] = ComponentToWorldMatrix.TransformPosition( EditableMesh.GetVertexAttribute( MeshVertexIDs[ EdgeVertexNumber ], UEditableMeshAttribute::VertexPosition(), 0 ) );
+	}
+
+	const bool bApplyDepthBias = true;
+	AddThickLineToDynamicMesh( CameraToWorld, EdgeVertexPositions, ColorAndOpacity, SizeBias, bApplyDepthBias, MeshBuilder );
+}
+
+
+void FMeshEditorMode::AddPolygonToDynamicMesh( const UEditableMesh& EditableMesh, const FTransform& CameraToWorld, const FMatrix& ComponentToWorldMatrix, const FPolygonRef PolygonRef, const FColor ColorAndOpacity, const float SizeBias, const bool bFillFaces, FDynamicMeshBuilder& VertexAndEdgeMeshBuilder, FDynamicMeshBuilder* PolygonFaceMeshBuilder )
+{
+	static TArray<FVertexID> MeshVertexIDs;
+	MeshVertexIDs.Reset();
+	EditableMesh.GetPolygonPerimeterVertices( PolygonRef, /* Out */ MeshVertexIDs );
+
+	static TArray<FVector> PolygonPerimeterVertexPositions;
+	PolygonPerimeterVertexPositions.SetNumUninitialized( MeshVertexIDs.Num(), false );
+	for( uint32 PolygonVertexIndex = 0; PolygonVertexIndex < (uint32)MeshVertexIDs.Num(); ++PolygonVertexIndex )
+	{
+		PolygonPerimeterVertexPositions[ PolygonVertexIndex ] = ComponentToWorldMatrix.TransformPosition( EditableMesh.GetVertexAttribute( MeshVertexIDs[ PolygonVertexIndex ], UEditableMeshAttribute::VertexPosition(), 0 ) );
+	}
+
+	// @todo mesheditor: Support polygon hole contours!
+
+
+	// Find the normal of the face (expects counter-clockwise winding). Flip the sign with the component to world determinant in case the component has been mirrored.
+	// @todo mesheditor: For polygon, assuming planar here for normal.  Not right.
+	const FVector PolygonNormal = FVector::CrossProduct( ( PolygonPerimeterVertexPositions[ 2 ] - PolygonPerimeterVertexPositions[ 0 ] ), ( PolygonPerimeterVertexPositions[ 1 ] - PolygonPerimeterVertexPositions[ 0 ] ) ).GetSafeNormal();
+	const float Determinant = ComponentToWorldMatrix.Determinant();
+
+	if( FVector::DotProduct( PolygonNormal, CameraToWorld.GetLocation() - PolygonPerimeterVertexPositions[ 0 ] ) * Determinant < 0.0f )
+	{
+		// Ignore backfaced polys
+		return;
+	}
+
+	if( bFillFaces )
+	{
+		check( PolygonFaceMeshBuilder );
+		const uint32 PolygonTriangleCount = EditableMesh.GetPolygonTriangulatedTriangleCount( PolygonRef );
+		for( uint32 PolygonTriangleNumber = 0; PolygonTriangleNumber < PolygonTriangleCount; ++PolygonTriangleNumber )
+		{
+			FVector TriangleVertexPositions[ 3 ];
+			FVector VertexOffsets[ 3 ];
+			for( uint32 TriangleVertexNumber = 0; TriangleVertexNumber < 3; ++TriangleVertexNumber )
+			{
+				TriangleVertexPositions[ TriangleVertexNumber ] =
+					ComponentToWorldMatrix.TransformPosition( 
+						EditableMesh.GetPolygonTriangulatedTriangleVertexPosition( PolygonRef, PolygonTriangleNumber, TriangleVertexNumber ) );
+
+				// We're offsetting the geometry from the actual face a bit, to avoid z-fighting for this particular effect
+				const FVector DirectionToCamera = CameraToWorld.GetLocation() - TriangleVertexPositions[ TriangleVertexNumber ];
+				const float DistanceToCamera = DirectionToCamera.Size();
+				const float DistanceBasedScaling = MeshEd::OverlayDistanceBias->GetFloat() + DistanceToCamera * MeshEd::OverlayDistanceScaleFactor->GetFloat();
+				VertexOffsets[ TriangleVertexNumber ] = ( DirectionToCamera / DistanceToCamera ) * MeshEd::OverlayDepthOffset->GetFloat() * DistanceBasedScaling;
+			}
+
+			// @todo mesheditor perf: Horribly inefficient way to draw a single triangle!  Also we should have a global utility function for this rather than using FDynamicMeshBuilder
+			// @todo mesheditor: Assumes vertices are in-order and form a single closed polygon.  Not accounting for holes.
+
+			const int32 FirstVertexIndex = 
+				PolygonFaceMeshBuilder->AddVertex( VertexOffsets[ 0 ] + TriangleVertexPositions[ 0 ], FVector2D( 0, 0 ), FVector( 1, 0, 0 ), FVector( 0, 1, 0 ), FVector( 0, 0, 1 ), ColorAndOpacity );
+			PolygonFaceMeshBuilder->AddVertex( VertexOffsets[ 1 ] + TriangleVertexPositions[ 1 ], FVector2D( 0, 1 ), FVector( 1, 0, 0 ), FVector( 0, 1, 0 ), FVector( 0, 0, 1 ), ColorAndOpacity );
+			PolygonFaceMeshBuilder->AddVertex( VertexOffsets[ 2 ] + TriangleVertexPositions[ 2 ], FVector2D( 1, 1 ), FVector( 1, 0, 0 ), FVector( 0, 1, 0 ), FVector( 0, 0, 1 ), ColorAndOpacity );
+
+			PolygonFaceMeshBuilder->AddTriangle( FirstVertexIndex + 0, FirstVertexIndex + 1, FirstVertexIndex + 2 );
+		}
+	}
+
+	{
+		static TArray<FEdgeID> PerimeterEdges;
+		EditableMesh.GetPolygonPerimeterEdges( PolygonRef, /* Out */ PerimeterEdges );
+
+		// Draw polygon edges
+		// @todo mesheditor: Assumes no holes, etc.
+		for( const FEdgeID EdgeID : PerimeterEdges )
+		{
+			AddEdgeToDynamicMesh( EditableMesh, CameraToWorld, ComponentToWorldMatrix, EdgeID, ColorAndOpacity, SizeBias, VertexAndEdgeMeshBuilder );
+		}
+	}
+}
+
+
+void FMeshEditorMode::DrawMeshElements( const FTransform& CameraToWorld, FViewport* Viewport, FPrimitiveDrawInterface* PDI, const TArrayView<FMeshElement>& MeshElements, const FColor Color, const bool bFillFaces, const float SizeBias, const TArray<FColor>* OptionalPerElementColors, const TArray<float>* OptionalPerElementSizeBiases )
+{
+	if( MeshElements.Num() > 0 )
+	{
+		// Must specify the same number of element colors and size biases as elements to draw
+		check( OptionalPerElementColors == nullptr || OptionalPerElementColors->Num() == MeshElements.Num() );
+		check( OptionalPerElementSizeBiases == nullptr || OptionalPerElementSizeBiases->Num() == MeshElements.Num() );
+
+		FDynamicMeshBuilder VertexAndEdgeMeshBuilder;
+		FDynamicMeshBuilder PolygonFaceMeshBuilder;
+
+		uint32 NumPolygonsToDraw = 0;
+		
+		for( int32 MeshElementIndex = 0; MeshElementIndex < MeshElements.Num(); ++MeshElementIndex )
+		{
+			const FMeshElement& MeshElement = MeshElements[ MeshElementIndex ];
+			if( MeshElement.IsValidMeshElement() )
+			{
+				const FColor ElementColor = OptionalPerElementColors != nullptr ? ( *OptionalPerElementColors )[ MeshElementIndex ] : Color;
+				const float ElementSizeBias = OptionalPerElementSizeBiases != nullptr ? ( *OptionalPerElementSizeBiases )[ MeshElementIndex ] : SizeBias;
+
+				UEditableMesh* EditableMesh = FindOrCreateEditableMesh( *MeshElement.Component, MeshElement.ElementAddress.SubMeshAddress );
+				if( EditableMesh != nullptr )
+				{
+					if( IsElementIDValid( MeshElement, EditableMesh ) )
+					{
+						UPrimitiveComponent* Component = MeshElement.Component.Get();
+						check( Component != nullptr )
+
+						const FMatrix ComponentToWorldMatrix = Component->GetRenderMatrix();
+
+						if( MeshElement.ElementAddress.ElementType == EEditableMeshElementType::Vertex )
+						{
+							const bool bApplyDepthBias = true;
+							AddVertexToDynamicMesh(
+								*EditableMesh,
+								CameraToWorld,
+								ComponentToWorldMatrix,
+								FVertexID( MeshElement.ElementAddress.ElementID ),
+								ElementColor,
+								ElementSizeBias,
+								bApplyDepthBias,
+								VertexAndEdgeMeshBuilder );
+						}
+						else if( MeshElement.ElementAddress.ElementType == EEditableMeshElementType::Edge )
+						{
+							AddEdgeToDynamicMesh(
+								*EditableMesh,
+								CameraToWorld,
+								ComponentToWorldMatrix,
+								FEdgeID( MeshElement.ElementAddress.ElementID ),
+								ElementColor,
+								ElementSizeBias,
+								VertexAndEdgeMeshBuilder );
+						}
+						else if( MeshElement.ElementAddress.ElementType == EEditableMeshElementType::Polygon )
+						{
+							++NumPolygonsToDraw;
+
+							AddPolygonToDynamicMesh(
+								*EditableMesh,
+								CameraToWorld,
+								ComponentToWorldMatrix,
+								FPolygonRef( MeshElement.ElementAddress.SectionID, FPolygonID( MeshElement.ElementAddress.ElementID ) ),
+								ElementColor,
+								ElementSizeBias,
+								bFillFaces,
+								VertexAndEdgeMeshBuilder,
+								&PolygonFaceMeshBuilder );
+						}
+					}
+				}
+			}
+		}
+
+		if( NumPolygonsToDraw && bFillFaces )
+		{
+			const bool bIsSelected = false;
+			FMaterialRenderProxy* MaterialRenderProxy = HoveredFaceMaterial->GetRenderProxy( bIsSelected );
+
+			const bool bDisableBackfaceCulling = true;	// @todo mesheditor: Disabling backface so we can see what's hovered even if its backfacing
+			const bool bReceivesDecals = false;
+			const FHitProxyId HitProxyId = FHitProxyId::InvisibleHitProxyId;	// Don't obscure selection!
+			PolygonFaceMeshBuilder.Draw( PDI, FMatrix::Identity, MaterialRenderProxy, SDPG_World, bDisableBackfaceCulling, bReceivesDecals, HitProxyId );
+		}
+
+		{
+			const bool bIsSelected = false;
+			FMaterialRenderProxy* MaterialRenderProxy = HoveredGeometryMaterial->GetRenderProxy( bIsSelected );
+
+			const bool bDisableBackfaceCulling = true;	// @todo mesheditor: Disabling backface so we can see what's hovered even if its backfacing
+			const bool bReceivesDecals = false;
+			const FHitProxyId HitProxyId = FHitProxyId::InvisibleHitProxyId;	// Don't obscure selection!
+			VertexAndEdgeMeshBuilder.Draw( PDI, FMatrix::Identity, MaterialRenderProxy, SDPG_World, bDisableBackfaceCulling, bReceivesDecals, HitProxyId );
+		}
+	}
+}
+
+
+void FMeshEditorMode::Render( const FSceneView* SceneView, FViewport* Viewport, FPrimitiveDrawInterface* PDI )
+{
+	FEdMode::Render( SceneView, Viewport, PDI );
+
+	const FTransform CameraToWorld = CachedCameraToWorld.IsSet() ? CachedCameraToWorld.GetValue() : FTransform( SceneView->ViewRotation, SceneView->ViewLocation );
+
+	// @todo mesheditor debug
+	if( MeshEd::ShowDebugStats->GetInt() > 0 && SelectedMeshElements.Num() > 0 )
+	{
+		FMeshElement& MeshElement = SelectedMeshElements[ 0 ];
+		UEditableMesh* EditableMesh = FindOrCreateEditableMesh( *MeshElement.Component, MeshElement.ElementAddress.SubMeshAddress );
+		if( EditableMesh != nullptr )
+		{
+			GEngine->AddOnScreenDebugMessage( INDEX_NONE, 0.0f, FColor::White,
+				FString::Printf( TEXT( "MeshElement: %s" ), *MeshElement.ToString() ), false );
+			GEngine->AddOnScreenDebugMessage( INDEX_NONE, 0.0f, FColor::White,
+				FString::Printf( TEXT( "Vertices: %i (array size: %i)" ), EditableMesh->GetVertexCount(), EditableMesh->GetVertexArraySize() ), false );
+			GEngine->AddOnScreenDebugMessage( INDEX_NONE, 0.0f, FColor::White,
+				FString::Printf( TEXT( "Rendering Vertices: %i (array size: %i)" ), EditableMesh->GetRenderingVertexCount(), EditableMesh->GetRenderingVertexArraySize() ), false );
+			GEngine->AddOnScreenDebugMessage( INDEX_NONE, 0.0f, FColor::White,
+				FString::Printf( TEXT( "Edges: %i (array size: %i)" ), EditableMesh->GetEdgeCount(), EditableMesh->GetEdgeArraySize() ), false );
+			GEngine->AddOnScreenDebugMessage( INDEX_NONE, 0.0f, FColor::White,
+				FString::Printf( TEXT( "Sections: %i (array size: %i)" ), EditableMesh->GetSectionCount(), EditableMesh->GetSectionArraySize() ), false );
+			for( int32 SectionIndex = 0; SectionIndex < EditableMesh->GetSectionArraySize(); ++SectionIndex )
+			{
+				const FSectionID SectionID( SectionIndex );
+				if( EditableMesh->IsValidSection( SectionID ) )
+				{
+					GEngine->AddOnScreenDebugMessage( INDEX_NONE, 0.0f, FColor::White,
+						FString::Printf( TEXT( "   [%i] Polygons: %i (array size: %i)" ), SectionID.GetValue(), EditableMesh->GetPolygonCount( SectionID ), EditableMesh->GetPolygonArraySize( SectionID ) ), false );
+					GEngine->AddOnScreenDebugMessage( INDEX_NONE, 0.0f, FColor::White,
+						FString::Printf( TEXT( "   [%i] Triangles: %i (array size: %i)" ), SectionID.GetValue(), EditableMesh->GetTriangleCount( SectionID ), EditableMesh->GetTriangleArraySize( SectionID ) ), false );
+				}
+			}
+		}
+	}
+
+	// Draw all polygon edges for selected/hovered meshes (@todo mesheditor: This only draws polygons right now.  We should also draw "floating" edges and vertices too!)
+	{
+		struct FComponentAndEditableMesh
+		{
+			const UPrimitiveComponent* Component;
+			const UEditableMesh* EditableMesh;
+
+			FComponentAndEditableMesh( const UPrimitiveComponent* InitComponent, const UEditableMesh* InitEditableMesh )
+				: Component( InitComponent ),
+				  EditableMesh( InitEditableMesh )
+			{
+			}
+
+			inline bool operator==( const FComponentAndEditableMesh& Other ) const
+			{
+				return Component == Other.Component && EditableMesh == Other.EditableMesh;
+			}
+		};
+
+
+		static TArray<FComponentAndEditableMesh> HoveredOrSelectedMeshes;
+		HoveredOrSelectedMeshes.Reset();
+
+		for( const FMeshElement& SelectedMeshElement : SelectedMeshElements )	// @todo mesheditor: Should also include meshes that are editor-selected
+		{
+			if( SelectedMeshElement.IsValidMeshElement() )
+			{
+				UEditableMesh* EditableMesh = FindOrCreateEditableMesh( *SelectedMeshElement.Component, SelectedMeshElement.ElementAddress.SubMeshAddress );
+				if( EditableMesh != nullptr )
+				{
+					HoveredOrSelectedMeshes.AddUnique( FComponentAndEditableMesh( SelectedMeshElement.Component.Get(), EditableMesh ) );
+				}
+			}
+		}
+
+		// Only draw hover if we're not in the middle of an interactive edit
+		if( ActiveAction == EMeshEditAction::None )
+		{
+			for( FMeshEditorInteractorData& MeshEditorInteractorData : MeshEditorInteractorDatas )
+			{
+				if( MeshEditorInteractorData.HoveredMeshElement.IsValidMeshElement() )
+				{
+					UEditableMesh* EditableMesh = FindOrCreateEditableMesh( *MeshEditorInteractorData.HoveredMeshElement.Component, MeshEditorInteractorData.HoveredMeshElement.ElementAddress.SubMeshAddress );
+					if( EditableMesh != nullptr )
+					{
+						HoveredOrSelectedMeshes.AddUnique( FComponentAndEditableMesh( MeshEditorInteractorData.HoveredMeshElement.Component.Get(), EditableMesh ) );
+					}
+				}
+			}
+		}
+
+		for( const FComponentAndEditableMesh& ComponentAndEditableMesh : HoveredOrSelectedMeshes )
+		{
+			const UPrimitiveComponent* Component = ComponentAndEditableMesh.Component;
+			check( Component != nullptr );
+
+			const UEditableMesh* EditableMesh = ComponentAndEditableMesh.EditableMesh;
+			check( EditableMesh != nullptr );
+
+			const FMatrix ComponentToWorldMatrix = Component->GetRenderMatrix();
+
+			FDynamicMeshBuilder VertexAndEdgeMeshBuilder;
+
+			// Draw subdivision limit surface if subdivision preview is enabled
+			if( EditableMesh->IsPreviewingSubdivisions() )
+			{
+				const float SizeBias = 0.0f;
+				const FColor SubdivisionEdgeColorAndOpacity = FLinearColor( 0.05f, 0.05f, 0.05f, 0.6f ).ToFColor( false );
+				const FColor BaseCageCounterpartEdgeColorAndOpacity = FLinearColor( 0.0f, 0.0f, 0.2f, 0.8f ).ToFColor( false );
+
+				const FSubdivisionLimitData& SubdivisionLimitData = EditableMesh->GetSubdivisionLimitData();
+
+				for( int32 WireEdgeNumber = 0; WireEdgeNumber < SubdivisionLimitData.SubdividedWireEdges.Num(); ++WireEdgeNumber )
+				{
+					const FSubdividedWireEdge& SubdividedWireEdge = SubdivisionLimitData.SubdividedWireEdges[ WireEdgeNumber ];
+
+					const int32 EdgeVertexIndexA = SubdividedWireEdge.EdgeVertex0PositionIndex;
+					const int32 EdgeVertexIndexB = SubdividedWireEdge.EdgeVertex1PositionIndex;
+
+					FVector EdgeVertexPositions[ 2 ];
+					EdgeVertexPositions[ 0 ] = SubdivisionLimitData.VertexPositions[ EdgeVertexIndexA ];
+					EdgeVertexPositions[ 1 ] = SubdivisionLimitData.VertexPositions[ EdgeVertexIndexB ];
+
+					FVector WorldSpaceEdgeVertexPositions[ 2 ];
+					for( uint32 EdgeVertexNumber = 0; EdgeVertexNumber < ARRAY_COUNT( EdgeVertexPositions ); ++EdgeVertexNumber )
+					{
+						WorldSpaceEdgeVertexPositions[ EdgeVertexNumber ] = ComponentToWorldMatrix.TransformPosition( EdgeVertexPositions[ EdgeVertexNumber ] );
+					}
+
+					const FColor ColorAndOpacity = SubdividedWireEdge.bIsBaseCageCounterpartEdge ? BaseCageCounterpartEdgeColorAndOpacity : SubdivisionEdgeColorAndOpacity;
+
+					const bool bApplyDepthBias = true;
+					AddThickLineToDynamicMesh( CameraToWorld, WorldSpaceEdgeVertexPositions, ColorAndOpacity, SizeBias, bApplyDepthBias, VertexAndEdgeMeshBuilder );
+				}
+			}
+
+			// Draw polygon mesh wires (or subdivision base cage, if previewing subdivisions)
+			{
+				const float Opacity = 0.85f;
+				const float SizeBias = 0.05f;
+				const bool bFillFaces = false;
+				const FColor ColorAndOpacity = FLinearColor( 0.0f, 0.0f, 0.3f ).CopyWithNewOpacity( Opacity ).ToFColor( false );
+
+				for( int32 SectionIndex = 0; SectionIndex < EditableMesh->GetSectionArraySize(); ++SectionIndex )
+				{
+					const FSectionID SectionID( SectionIndex );
+					if( EditableMesh->IsValidSection( SectionID ) )
+					{
+						const uint32 SectionPolygonArraySize = EditableMesh->GetPolygonArraySize( SectionID );
+
+						for( uint32 PolygonIndex = 0; PolygonIndex < SectionPolygonArraySize; ++PolygonIndex )
+						{
+							const FPolygonID PolygonID( PolygonIndex );
+
+							const FPolygonRef PolygonRef( SectionID, PolygonID );
+							if( EditableMesh->IsValidPolygon( PolygonRef ) )
+							{
+								AddPolygonToDynamicMesh(
+									*EditableMesh,
+									CameraToWorld,
+									ComponentToWorldMatrix,
+									PolygonRef,
+									ColorAndOpacity,
+									SizeBias,
+									bFillFaces,
+									VertexAndEdgeMeshBuilder,
+									nullptr );
+
+								if( bShowVertexNormals )
+								{
+									// @todo mesheditor: total debug feature for now. Need a way of making this look nice.
+									const int32 PerimeterVertexCount = EditableMesh->GetPolygonPerimeterVertexCount( PolygonRef );
+									const float Length = 10.0f; // @todo mesheditor: determine length of debug line from distance from the mesh origin to the camera?
+
+									for( int32 PerimeterVertexIndex = 0; PerimeterVertexIndex < PerimeterVertexCount; ++PerimeterVertexIndex )
+									{
+										FVector Position( EditableMesh->GetPolygonPerimeterVertexAttribute( PolygonRef, PerimeterVertexIndex, UEditableMeshAttribute::VertexPosition(), 0 ) );
+										FVector Normal( EditableMesh->GetPolygonPerimeterVertexAttribute( PolygonRef, PerimeterVertexIndex, UEditableMeshAttribute::VertexNormal(), 0 ) );
+
+										FVector VertexNormalRenderPositions[] =
+										{
+											ComponentToWorldMatrix.TransformPosition( Position ),
+											ComponentToWorldMatrix.TransformPosition( Position + Normal * Length )
+										};
+
+										// We don't need any depth bias when drawing lines for face normals
+										const bool bApplyDepthBias = false;
+										AddThickLineToDynamicMesh(
+											CameraToWorld,
+											VertexNormalRenderPositions,
+											FColor::Magenta,
+											SizeBias,
+											bApplyDepthBias,
+											VertexAndEdgeMeshBuilder );
+									}
+								}
+							}
+						}
+					}
+				}
+			}
+
+			const bool bIsSelected = false;
+			FMaterialRenderProxy* MaterialRenderProxy = HoveredGeometryMaterial->GetRenderProxy( bIsSelected );
+
+			const bool bDisableBackfaceCulling = true;	// @todo mesheditor: Disabling backface so we can see what's hovered even if its backfacing
+			const bool bReceivesDecals = false;
+			const FHitProxyId HitProxyId = FHitProxyId::InvisibleHitProxyId;	// Don't obscure selection!
+			VertexAndEdgeMeshBuilder.Draw( PDI, FMatrix::Identity, MaterialRenderProxy, SDPG_World, bDisableBackfaceCulling, bReceivesDecals, HitProxyId );
+		}
+	}
+
+
+	// Draw hovered and selected elements
+	{
+		// Only draw hover if we're not in the middle of an interactive edit
+		if( ActiveAction == EMeshEditAction::None )
+		{
+			const float HoveredSizeBias = MeshEd::HoveredSizeBias->GetFloat() + MeshEd::HoveredAnimationExtraSizeBias->GetFloat() * FMath::MakePulsatingValue( HoverFeedbackTimeValue, 0.5f );
+			{
+				static TArray<FMeshElement> HoveredMeshElementsToDraw;
+				HoveredMeshElementsToDraw.Reset();
+
+				// Draw hovered meshes
+				for( FMeshEditorInteractorData& MeshEditorInteractorData : MeshEditorInteractorDatas )
+				{
+					if( MeshElementSelectionMode == EEditableMeshElementType::Any ||
+					    MeshEditorInteractorData.HoveredMeshElement.ElementAddress.ElementType == MeshElementSelectionMode )
+					{
+						HoveredMeshElementsToDraw.Add( MeshEditorInteractorData.HoveredMeshElement );
+					}
+				}
+
+				const float Opacity = 1.0f;
+				const bool bFillFaces = true;
+				DrawMeshElements(
+					CameraToWorld,
+					Viewport,
+					PDI,
+					HoveredMeshElementsToDraw,
+					FLinearColor::Green.CopyWithNewOpacity( Opacity ).ToFColor( false ),
+					bFillFaces,
+					HoveredSizeBias );
+			}
+
+			// Draw meshes that were previously hovered
+			{
+				static TArray<FColor> PerElementColors;
+				PerElementColors.Reset();
+
+				const double CurrentRealTime = FSlateApplication::Get().GetCurrentTime();
+
+				const float HoverFadeTime = MeshEd::HoverFadeDuration->GetFloat();;
+				for( FMeshElement& FadingOutHoveredMeshElement : FadingOutHoveredMeshElements )
+				{
+					const float TimeSinceLastHovered = CurrentRealTime - FadingOutHoveredMeshElement.LastHoverTime;
+					float Opacity = 1.0f - ( TimeSinceLastHovered / HoverFadeTime );
+					Opacity = Opacity * Opacity * Opacity * Opacity;		// Exponential falloff
+					Opacity = FMath::Clamp( Opacity, 0.0f, 1.0f );
+
+					PerElementColors.Add( FLinearColor::Green.CopyWithNewOpacity( Opacity ).ToFColor( false ) );
+				}
+
+				const bool bFillFaces = true;
+				DrawMeshElements(
+					CameraToWorld,
+					Viewport,
+					PDI,
+					FadingOutHoveredMeshElements,
+					FColor::White,	// Ignored, as we'll pass in per-element colors also
+					bFillFaces,
+					HoveredSizeBias,
+					&PerElementColors );
+			}
+		}
+
+		// Draw selected mesh elements
+		{
+			const double CurrentRealTime = FSlateApplication::Get().GetCurrentTime();
+
+			static TArray<float> PerElementSizeBiases;
+			PerElementSizeBiases.Reset();
+
+			const float SelectionAnimationDuration = MeshEd::SelectionAnimationDuration->GetFloat();
+			for( FMeshElement& SelectedMeshElement : SelectedMeshElements )
+			{
+				const float TimeSinceSelected = CurrentRealTime - SelectedMeshElement.LastSelectTime;
+
+				const float SizeBias = MeshEd::SelectedSizeBias->GetFloat() + MeshEd::SelectedAnimationExtraSizeBias->GetFloat() * FMath::Clamp( 1.0f - ( TimeSinceSelected / SelectionAnimationDuration ), 0.0f, 1.0f );
+				PerElementSizeBiases.Add( SizeBias );
+			}
+
+			const float Opacity = 1.0f;
+			const bool bFillFaces = true;
+			DrawMeshElements(
+				CameraToWorld,
+				Viewport,
+				PDI,
+				SelectedMeshElements,
+				FLinearColor::White.CopyWithNewOpacity( Opacity ).ToFColor( false ),
+				bFillFaces,
+				MeshEd::SelectedSizeBias->GetFloat(),
+				nullptr,
+				&PerElementSizeBiases );
+		}
+	}
+}
+
+
+FMeshEditorMode::FMeshEditorInteractorData& FMeshEditorMode::GetMeshEditorInteractorData( UViewportInteractor* ViewportInteractor )
+{
+	check( ViewportInteractor != nullptr );
+
+	// @todo mesheditor perf: We could use a hash table here for a faster lookup, but it's unlikely there will be more than a handful
+	// of viewport interactors to iterate over.
+	FMeshEditorInteractorData* FoundMeshEditorInteractorData = nullptr;
+	for( int32 MeshEditorInteractorDataIndex = 0; MeshEditorInteractorDataIndex < MeshEditorInteractorDatas.Num(); ++MeshEditorInteractorDataIndex )
+	{
+		FMeshEditorInteractorData& MeshEditorInteractorData = MeshEditorInteractorDatas[ MeshEditorInteractorDataIndex ];
+
+		UViewportInteractor* CurrentViewportInteractor = MeshEditorInteractorData.ViewportInteractor.Get();
+		if( CurrentViewportInteractor == nullptr )
+		{
+			// Expired
+			MeshEditorInteractorDatas.RemoveAt( MeshEditorInteractorDataIndex-- );
+		}
+		else
+		{
+			if( CurrentViewportInteractor == ViewportInteractor )
+			{
+				FoundMeshEditorInteractorData = &MeshEditorInteractorData;
+			}
+		}
+	}
+
+	if( FoundMeshEditorInteractorData == nullptr )
+	{
+		FMeshEditorInteractorData& NewMeshEditorInteractorData = *new( MeshEditorInteractorDatas ) FMeshEditorInteractorData();
+		NewMeshEditorInteractorData.ViewportInteractor = ViewportInteractor;
+		FoundMeshEditorInteractorData = &NewMeshEditorInteractorData;
+	}
+
+	check( FoundMeshEditorInteractorData != nullptr );
+	return *FoundMeshEditorInteractorData;
+}
+
+
+// @todo mesheditor debug
+// static UViewportWorldInteraction* GHackVWI = nullptr;
+// static FTransform GHackComponentToWorld;
+
+void FMeshEditorMode::OnViewportInteractionHoverUpdate( UViewportInteractor* ViewportInteractor, FVector& OutHoverImpactPoint, bool& bWasHandled )
+{
+	if( !bWasHandled )
+	{
+		FMeshEditorInteractorData& MeshEditorInteractorData = GetMeshEditorInteractorData( ViewportInteractor );
+
+		MeshEditorInteractorData.bGrabberSphereIsValid = ViewportInteractor->GetGrabberSphere( /* Out */ MeshEditorInteractorData.GrabberSphere );
+		MeshEditorInteractorData.bLaserIsValid = ViewportInteractor->GetLaserPointer( /* Out */ MeshEditorInteractorData.LaserStart, /* Out */ MeshEditorInteractorData.LaserEnd );
+
+		// @todo gizmo: Hover should have already been cleared earlier this frame when Tick() was called, but in some cases when clicking, tick won't be called
+		//check( !MeshEditorInteractorData.HoveredMeshElement.IsValidMeshElement() );
+
+		const int32 LODIndex = 0;			// @todo mesheditor: We'll want to select an LOD to edit in various different wants (LOD that's visible, or manual user select, etc.)
+											
+		// Don't use the laser pointer while if someone else has captured input
+		// @todo vreditor: We need to re-think how input capture works.  This seems too hacky/complex
+		FViewportActionKeyInput* SelectAndMoveAction = ViewportInteractor->GetActionWithName( ViewportWorldActionTypes::SelectAndMove );
+		FViewportActionKeyInput* WorldMovementAction = ViewportInteractor->GetActionWithName( ViewportWorldActionTypes::WorldMovement );
+ 		const bool bIsLaserPointerBusy = 
+			( SelectAndMoveAction != nullptr && SelectAndMoveAction->bIsInputCaptured && ActiveAction == EMeshEditAction::None ) ||
+			( WorldMovementAction != nullptr && WorldMovementAction->bIsInputCaptured && ActiveAction == EMeshEditAction::None );
+
+		bool bIsGrabberSphereOverMeshElement = false;
+
+		if( !bIsLaserPointerBusy )
+		{
+			if( ActiveAction == EMeshEditAction::None || 
+				bActiveActionNeedsHoverLocation )
+			{
+				const float WorldSpaceRayFuzzyDistance = MeshEd::LaserFuzzySelectionDistance->GetFloat() * ViewportWorldInteraction->GetWorldScaleFactor();
+				const float WorldSpaceGrabberSphereFuzzyDistance = MeshEd::GrabberSphereFuzzySelectionDistance->GetFloat() * ViewportWorldInteraction->GetWorldScaleFactor();
+				const float ExtraFuzzyScalingForCollisionQuery = 1.25f;	// @todo mesheditor urgent: Inflates collision query bounds to account for us not doing any distance-based scaling of query size
+
+				// Two passes -- first with grabber sphere, then again with the laser
+				const int32 FirstInteractorPassNumber = GetDefault<UMeshEditorSettings>()->bAllowGrabberSphere ? 0 : 1;
+				for( int32 InteractorPassNumber = FirstInteractorPassNumber; InteractorPassNumber < 2; ++InteractorPassNumber )
+				{
+					const bool bIsGrabberSphereTest = ( InteractorPassNumber == 0 );
+					const bool bIsLaserTest = !bIsGrabberSphereTest;
+
+					if( ( bIsGrabberSphereTest && MeshEditorInteractorData.bGrabberSphereIsValid ) ||
+						( bIsLaserTest && MeshEditorInteractorData.bLaserIsValid ) )
+					{
+						static TArray< UPrimitiveComponent* > HitComponents;
+						HitComponents.Reset();
+
+						// Trace against the world twice.  Once for simple collision and then again for complex collision.
+						// We need the simple collision pass so that we can catch editable meshes with inflated bounds for
+						// subdivision cages.
+						for( int32 CollisionPassNumber = 0; CollisionPassNumber < 2; ++CollisionPassNumber )
+						{
+							const bool bTraceComplex = ( CollisionPassNumber == 0 );
+							FCollisionQueryParams TraceParams( NAME_None, bTraceComplex, nullptr );
+
+							static TArray<UPrimitiveComponent*> ComponentsFoundThisPass;
+							ComponentsFoundThisPass.Reset();
+
+							if( bIsGrabberSphereTest )
+							{
+								// Grabber sphere testing
+								FCollisionShape CollisionShape;
+								CollisionShape.SetSphere( MeshEditorInteractorData.GrabberSphere.W + WorldSpaceGrabberSphereFuzzyDistance * ExtraFuzzyScalingForCollisionQuery );
+
+	 							//DrawDebugSphere( GetWorld(), MeshEditorInteractorData.GrabberSphere.Center, 1.5f * ViewportWorldInteraction->GetWorldScaleFactor(), 32, FColor::White, false, 0.0f );
+	 							//DrawDebugSphere( GetWorld(), MeshEditorInteractorData.GrabberSphere.Center, CollisionShape.GetSphereRadius(), 32, FColor( 255, 40, 40, 255 ), false, 0.0f );
+
+								static TArray< FOverlapResult > OverlapResults;
+								OverlapResults.Reset();
+								if( GetWorld()->OverlapMultiByChannel( OverlapResults, MeshEditorInteractorData.GrabberSphere.Center, FQuat::Identity, ECC_Visibility, CollisionShape, TraceParams ) )
+								{
+									for( FOverlapResult& OverlapResult : OverlapResults )
+									{
+										if( OverlapResult.GetComponent() != nullptr )
+										{
+											ComponentsFoundThisPass.Add( OverlapResult.GetComponent() );
+										}
+									}
+								}
+							}
+							else
+							{
+								// Fuzzy hit testing (thick laser)
+								FCollisionShape CollisionShape;
+								CollisionShape.SetSphere( WorldSpaceRayFuzzyDistance * ExtraFuzzyScalingForCollisionQuery );
+
+								// @todo mesheditor perf: This could be fairly slow, tracing so many objects.  We could do SweepSingleByChannel, but the nearest mesh might not
+								// actually have the best element to select
+								static TArray< FHitResult > HitResults;
+								HitResults.Reset();
+								if( GetWorld()->SweepMultiByChannel( HitResults, MeshEditorInteractorData.LaserStart, MeshEditorInteractorData.LaserEnd, FQuat::Identity, ECC_Visibility, CollisionShape, TraceParams ) )
+								{
+									for( FHitResult& HitResult : HitResults )
+									{
+										if( HitResult.GetComponent() != nullptr )
+										{
+											ComponentsFoundThisPass.Add( HitResult.GetComponent() );
+										}
+									}
+								}
+							}
+
+							// @todo mesheditor: We could avoid multiple collision test passes if the physics system had a way to do
+							// per-shape filtering during our query.  Basically we'd recognize the bounds shape for editable mesh's
+							// components and always trace against that instead of its complex collision
+							for( UPrimitiveComponent* Component : ComponentsFoundThisPass )
+							{
+								// Always add components we find from the complex collision pass, but only add meshes from the simple collision pass if 
+								// they're in subdivision preview mode.  Their base cage mesh won't match their complex collision geometry, but we
+								// still need to allow the user to interact with mesh elements outside the bounds of that geometry.
+								FEditableMeshSubMeshAddress SubMeshAddress = UEditableMeshFactory::MakeSubmeshAddress( Component, LODIndex );
+								UEditableMesh** EditableMeshPtr = CachedEditableMeshes.Find( SubMeshAddress );
+								if( bTraceComplex || ( EditableMeshPtr != nullptr && ( *EditableMeshPtr )->IsPreviewingSubdivisions() ) )
+								{
+									// Don't bother with editor-only 'helper' actors, we never want to visualize or edit geometry on those
+									if( !Component->IsEditorOnly() &&
+										( Component->GetOwner() == nullptr || !Component->GetOwner()->IsEditorOnly() ) )
+									{
+										HitComponents.AddUnique( Component );
+									}
+								}
+							}
+						}
+
+						// Find *everything* under the cursor, as well as the closest thing under the cursor.  This is so that
+						// systems can choose to filter out elements they aren't interested in.
+						// @todo mesheditor selection: We need to finish implementing this so that you can pick up something that's selected, even if the
+						// closest thing under the cursor isn't the same element type that you have selected.
+		// 						static TArray<FEditableMeshElementAddress> CandidateElementAddresses;
+		// 						CandidateElementAddresses.Reset();
+
+						UPrimitiveComponent* ClosestComponent = nullptr;
+						FEditableMeshElementAddress ClosestElementAddress;
+						EInteractorShape ClosestInteractorShape = EInteractorShape::Invalid;
+						FVector ClosestHoverLocation = FVector::ZeroVector;
+
+						for( UPrimitiveComponent* HitComponent : HitComponents )
+						{
+							const bool bOnlyInteractWithSelectedObjects = GetDefault<UMeshEditorSettings>()->bOnlyEditSelectedObjects;
+							if( !bOnlyInteractWithSelectedObjects || GEditor->GetSelectedActors()->IsSelected( HitComponent->GetOwner() ) )
+							{
+								// @todo mesheditor debug
+								// GHackComponentToWorld = HitComponent->GetComponentToWorld();
+								// GHackVWI = this->ViewportWorldInteraction;
+
+								const FMatrix& ComponentToWorldMatrix = HitComponent->GetRenderMatrix();
+								const float ComponentSpaceRayFuzzyDistance = ComponentToWorldMatrix.InverseTransformVector( FVector( WorldSpaceRayFuzzyDistance, 0.0f, 0.0f ) ).Size();
+								const float ComponentSpaceGrabberSphereFuzzyDistance = ComponentToWorldMatrix.InverseTransformVector( FVector( WorldSpaceGrabberSphereFuzzyDistance, 0.0f, 0.0f ) ).Size();
+
+								FEditableMeshSubMeshAddress SubMeshAddressToQuery = UEditableMeshFactory::MakeSubmeshAddress( HitComponent, LODIndex );
+
+								// Grab the existing editable mesh from our cache if we have one, otherwise create one now
+								UEditableMesh* EditableMesh = FindOrCreateEditableMesh( *HitComponent, SubMeshAddressToQuery );
+								if( EditableMesh != nullptr )
+								{
+									// If we're selecting by painting, only hover over elements of the same type that we already have selected
+									EEditableMeshElementType OnlyElementType = EEditableMeshElementType::Invalid;
+									if( this->MeshElementSelectionMode != EEditableMeshElementType::Any )
+									{
+										OnlyElementType = this->MeshElementSelectionMode;
+									}
+									else if( ActiveAction == EMeshEditAction::SelectByPainting )
+									{
+										OnlyElementType = GetSelectedMeshElementType();
+									}
+
+									// Shapes are in world space, but we need it in the local space of our component
+									const FVector ComponentSpaceLaserStart = ComponentToWorldMatrix.InverseTransformPosition( MeshEditorInteractorData.LaserStart );
+									const FVector ComponentSpaceLaserEnd = ComponentToWorldMatrix.InverseTransformPosition( MeshEditorInteractorData.LaserEnd );
+
+									const FSphere ComponentSpaceGrabberSphere(
+										ComponentToWorldMatrix.InverseTransformPosition( MeshEditorInteractorData.GrabberSphere.Center ),
+										ComponentToWorldMatrix.InverseTransformVector( FVector( MeshEditorInteractorData.GrabberSphere.W ) ).X );
+
+									const FTransform CameraToWorld = CachedCameraToWorld.IsSet() ? CachedCameraToWorld.GetValue() : HitComponent->GetComponentToWorld();	// @todo mesheditor: We're expecting to have a CameraToWorld here
+									const FVector ComponentSpaceCameraLocation = ComponentToWorldMatrix.InverseTransformPosition( CameraToWorld.GetLocation() );
+									const float ComponentSpaceFuzzyDistanceScaleFactor = ComponentToWorldMatrix.InverseTransformVector( FVector( MeshEd::OverlayDistanceScaleFactor->GetFloat() / ViewportWorldInteraction->GetWorldScaleFactor(), 0.0f, 0.0f ) ).Size();
+
+									EInteractorShape HitInteractorShape = EInteractorShape::Invalid;
+									FVector ComponentSpaceHitLocation = FVector::ZeroVector;
+									FEditableMeshElementAddress MeshElementAddress = QueryElement( 
+										*EditableMesh, 
+										bIsGrabberSphereTest, 
+										ComponentSpaceGrabberSphere, 
+										ComponentSpaceGrabberSphereFuzzyDistance, 
+										bIsLaserTest, 
+										ComponentSpaceLaserStart, 
+										ComponentSpaceLaserEnd, 
+										ComponentSpaceRayFuzzyDistance, 
+										OnlyElementType, 
+										ComponentSpaceCameraLocation,
+										ComponentSpaceFuzzyDistanceScaleFactor,
+										/* Out */ HitInteractorShape,
+										/* Out */ ComponentSpaceHitLocation );
+
+									if( MeshElementAddress.ElementType != EEditableMeshElementType::Invalid )
+									{
+										const FVector WorldSpaceHitLocation = ComponentToWorldMatrix.TransformPosition( ComponentSpaceHitLocation );
+										
+										const float ClosestDistanceToGrabberSphere = ( MeshEditorInteractorData.GrabberSphere.Center - ClosestHoverLocation ).Size();
+										const float DistanceToGrabberSphere = ( MeshEditorInteractorData.GrabberSphere.Center - WorldSpaceHitLocation ).Size();
+
+										const float ClosestDistanceOnRay = ( MeshEditorInteractorData.LaserStart - ClosestHoverLocation ).Size();
+										const float DistanceOnRay = ( MeshEditorInteractorData.LaserStart - WorldSpaceHitLocation ).Size();
+
+										// NOTE: We're preferring any grabber sphere hit over laser hits
+										if( ClosestComponent == nullptr ||
+											( HitInteractorShape == EInteractorShape::GrabberSphere && DistanceToGrabberSphere < ClosestDistanceToGrabberSphere ) ||
+											( HitInteractorShape == EInteractorShape::Laser && DistanceOnRay < ClosestDistanceOnRay ) )
+										{
+											ClosestComponent = HitComponent;
+											ClosestElementAddress = MeshElementAddress;
+											ClosestInteractorShape = HitInteractorShape;
+											ClosestHoverLocation = WorldSpaceHitLocation;
+										}
+									}
+								}
+							}
+						}
+
+						if( ClosestElementAddress.ElementType != EEditableMeshElementType::Invalid )
+						{
+							// We have a hovered element!
+							MeshEditorInteractorData.HoveredMeshElement.Component = ClosestComponent;
+							MeshEditorInteractorData.HoveredMeshElement.LastHoverTime = FSlateApplication::Get().GetCurrentTime();
+							MeshEditorInteractorData.HoveredMeshElement.ElementAddress = ClosestElementAddress;
+							MeshEditorInteractorData.HoverInteractorShape = ClosestInteractorShape;
+							MeshEditorInteractorData.HoverLocation = ClosestHoverLocation;
+
+							bWasHandled = true;
+							OutHoverImpactPoint = MeshEditorInteractorData.HoverLocation;
+
+							// If we hit something with our grabber sphere, then don't bother checking with the laser.  We always
+							// prefer grabber sphere hits.
+							if( bIsGrabberSphereTest )
+							{
+								//	DrawDebugSphere( GetWorld(), MeshEditorInteractorData.HoverLocation, 1.5f * ViewportWorldInteraction->GetWorldScaleFactor(), 16, FColor( 255, 40, 40, 255 ), false, 0.0f );
+								bIsGrabberSphereOverMeshElement = true;
+								break;
+							}
+						}
+					}
+				}
+			}
+		}
+
+		// Are we hovering over something new? (or nothing?)  If so, then we'll fade out the old hovered mesh element
+		const FMeshElement& PreviouslyHoveredMeshElement = MeshEditorInteractorData.PreviouslyHoveredMeshElement;
+		if( PreviouslyHoveredMeshElement.IsValidMeshElement() &&
+			!PreviouslyHoveredMeshElement.IsSameMeshElement( MeshEditorInteractorData.HoveredMeshElement ) )
+		{
+			// Replace any existing previously hovered element that points to the same mesh element
+			bool bAlreadyExisted = false;
+			for( int32 ElementIndex = 0; ElementIndex < FadingOutHoveredMeshElements.Num(); ++ElementIndex )
+			{
+				FMeshElement& ExistingElement = FadingOutHoveredMeshElements[ ElementIndex ];
+				if( ExistingElement.IsSameMeshElement( PreviouslyHoveredMeshElement ) )
+				{
+					ExistingElement = PreviouslyHoveredMeshElement;
+					bAlreadyExisted = true;
+					break;
+				}
+			}
+			if( !bAlreadyExisted )
+			{
+				if( MeshElementSelectionMode == EEditableMeshElementType::Any ||
+				    MeshElementSelectionMode == PreviouslyHoveredMeshElement.ElementAddress.ElementType )
+				{
+					FadingOutHoveredMeshElements.Add( PreviouslyHoveredMeshElement );
+				}
+			}
+		}
+	}
+}
+
+
+void FMeshEditorMode::OnViewportInteractionInputUnhandled( FEditorViewportClient& ViewportClient, class UViewportInteractor* ViewportInteractor, const struct FViewportActionKeyInput& Action )
+{
+	if( Action.ActionType == ViewportWorldActionTypes::SelectAndMove )
+	{
+		if( Action.Event == IE_Pressed )
+		{
+			// Deselect everything
+			if( SelectedMeshElements.Num() > 0 )
+			{
+				const FScopedTransaction Transaction( LOCTEXT( "UndoDeselectingAllMeshElements", "Deselect All Elements" ) );
+
+				StoreUndo( MeshEditorModeProxyObject, FDeselectAllMeshElementsChange( MoveTemp( FDeselectAllMeshElementsChangeInput() ) ).Execute( MeshEditorModeProxyObject ) );
+			}
+		}
+	}
+}
+
+
+void FMeshEditorMode::OnViewportInteractionStartDragging( UViewportInteractor* ViewportInteractor )
+{
+	if( ActiveAction == EMeshEditAction::None )
+	{
+		// NOTE: We pass an empty Undo text string to tell StartAction() that we don't need it to start a transaction
+		// because the caller of this delegate will have already done that (the viewport interaction system)
+		const bool bActionNeedsHoverLocation = false;
+		StartAction( EMeshEditAction::MoveUsingGizmo, ViewportInteractor, bActionNeedsHoverLocation, FText() );
+	}
+}
+
+
+void FMeshEditorMode::OnViewportInteractionStopDragging( UViewportInteractor* ViewportInteractor )
+{
+	// This will be called when the user releases the button/trigger to stop dragging, however the objects
+	// could still be moving after this is called.  This is because objects can interpolate to their final
+	// (snapped) position, or they could be "thrown" and inertia will carry them further.  To find out
+	// when the objects have finally stopped moving, check out OnViewportInteractionFinishedMovingTransformables()
+}
+
+
+void FMeshEditorMode::OnViewportInteractionFinishedMovingTransformables()
+{
+	if( ActiveAction != EMeshEditAction::None )
+	{
+		FinishAction();
+	}
+}
+
+
+void FMeshEditorMode::OnVREditorModePlaceDraggedMaterial( UPrimitiveComponent* HitComponent, UMaterialInterface* MaterialInterface, bool& bPlaced )
+{
+	if( !bPlaced )
+	{
+		static TMap< UEditableMesh*, TArray<FMeshElement> > MeshesAndPolygons;
+		GetSelectedMeshesAndPolygons( /* Out */ MeshesAndPolygons );
+		for( const auto& MeshAndPolygons : MeshesAndPolygons )
+		{
+			if( MeshAndPolygons.Value.Num() > 0 )
+			{
+				if( MeshAndPolygons.Value[ 0 ].Component.Get() == HitComponent )
+				{
+					AssignMaterialToSelectedPolygons( MaterialInterface );
+					bPlaced = true;
+					break;
+				}
+			}
+		}
+	}
+}
+
+
+void FMeshEditorMode::UpdateActiveAction( const bool bIsActionFinishing )
+{
+
+	// Make sure there are no outstanding changes being previewed -- we never want changes to STACK.  This can happen
+	// when UpdateActiveAction() is called more than once per frame.
+	RollbackPreviewChanges();
+
+	if( bIsFirstActiveActionUpdate &&
+	   ActiveAction != EMeshEditAction::None &&
+	   ActiveAction != EMeshEditAction::SelectByPainting )
+	{
+		CommitSelectedMeshes();
+	}
+
+
+	bool bShouldDeselectAllFirst = false;
+	TArray<FMeshElement> MeshElementsToSelect;
+
+	// @todo mesheditor urgent: During an interactive edit, if nothing ends up selected after the edit is complete, 
+	// no mesh elements will be rendered that frame, which makes it hard to see what's going on.  Currently, we make
+	// sure that something is always selected after every type of interactive edit, but in the future that may not make sense.
+
+
+
+	auto FindExtrudeDistanceUsingInteractor = [](
+		FMeshEditorMode* This,
+		UViewportInteractor* ViewportInteractor,
+		const UEditableMesh* EditableMesh,
+		const FVector AxisOrigin,
+		const FVector AxisDirection,
+		const float AxisLength,
+		float& OutExtrudeDistance )
+	{
+		check( ViewportInteractor != nullptr );
+
+		OutExtrudeDistance = 0.0f;
+
+		float ClosestDistanceToRay = MAX_flt;
+		float BestExtrudeDistance = 0.0f;
+		
+
+		// Find the interactor ray that is closest to the axis line, and determine the distance from the origin of the
+		// axis to that point
+		FMeshEditorInteractorData& MeshEditorInteractorData = This->GetMeshEditorInteractorData( ViewportInteractor );
+		if( MeshEditorInteractorData.bLaserIsValid )	// @todo grabber: Grabber sphere impl neeed
+		{
+			const FVector AxisSegmentStart = AxisOrigin - AxisDirection * ( AxisLength * 0.5f );
+			const FVector AxisSegmentEnd = AxisOrigin + AxisDirection * ( AxisLength * 0.5f );
+
+			FVector ClosestPointOnAxis, ClosestPointOnRay;
+			FMath::SegmentDistToSegmentSafe(
+				AxisSegmentStart, AxisSegmentEnd,
+				MeshEditorInteractorData.LaserStart, MeshEditorInteractorData.LaserEnd,
+				/* Out */ ClosestPointOnAxis,
+				/* Out */ ClosestPointOnRay );
+			const float DistanceToRay = ( ClosestPointOnAxis - ClosestPointOnRay ).Size();
+
+			if( DistanceToRay < ClosestDistanceToRay )
+			{
+				ClosestDistanceToRay = DistanceToRay;
+
+				const FVector AxisToClosestPointVector = ( ClosestPointOnAxis - AxisOrigin );
+				BestExtrudeDistance = AxisToClosestPointVector.Size();
+
+				// Check to see if the closest point is actually behind the origin of the axis (negative extrusion distance)
+				const FVector ClosestPointDirection = AxisToClosestPointVector.GetSafeNormal();
+				if( FVector::DotProduct( ClosestPointDirection, AxisDirection ) < 0.0f )
+				{
+					BestExtrudeDistance *= -1.0f;
+				}
+			}
+		}
+
+		OutExtrudeDistance = BestExtrudeDistance;
+	};
+
+	// Figures out how far we should inset on a polygon by figuring out the maximum distance from all of the polygon's perimeter
+	// vertices to the center of the polygon, then subtracting the distance between the hovered impact point and polygon center.
+	auto FindInsetAmount = [](
+		FMeshEditorMode* This,
+		UViewportInteractor* ViewportInteractor,
+		const FPolygonRef& PolygonRef,
+		UPrimitiveComponent& Component,
+		const UEditableMesh* EditableMesh,
+		float& OutInsetFixedDistance,
+		float& OutInsetProgressTowardCenter )
+	{
+		// @todo grabber: Glitches out when using grabber sphere near the corner of an inset polygon.
+		check( ViewportInteractor != nullptr );
+
+		OutInsetFixedDistance = 0.0f;
+		OutInsetProgressTowardCenter = 0.0f;
+
+		FMeshEditorInteractorData& MeshEditorInteractorData = This->GetMeshEditorInteractorData( ViewportInteractor );
+		if( MeshEditorInteractorData.bLaserIsValid || MeshEditorInteractorData.bGrabberSphereIsValid )
+		{
+			const FMatrix ComponentToWorldMatrix = Component.GetRenderMatrix();
+
+			const FPlane PolygonPlane = EditableMesh->ComputePolygonPlane( PolygonRef );
+			const FVector PolygonCenter = EditableMesh->ComputePolygonCenter( PolygonRef );
+
+			// @todo mesheditor: no longer supports grabber sphere. Add this back in (plane vs sphere closest point check) if we want to support it.
+			const FVector ComponentSpaceRayStart = ComponentToWorldMatrix.InverseTransformPosition( MeshEditorInteractorData.LaserStart );
+			const FVector ComponentSpaceRayEnd = ComponentToWorldMatrix.InverseTransformPosition( MeshEditorInteractorData.LaserEnd );
+
+			FVector ComponentSpaceRayIntersectionWithPolygonPlane;
+			if( FMath::SegmentPlaneIntersection( ComponentSpaceRayStart, ComponentSpaceRayEnd, PolygonPlane, /* Out */ ComponentSpaceRayIntersectionWithPolygonPlane ) )
+			{
+				// @todo mesheditor debug
+				// DrawDebugSphere( This->GetWorld(), ComponentToWorldMatrix.TransformPosition( ComponentSpaceRayIntersectionWithPolygonPlane ), 1.5f, 32, FColor::Red, false, 0.0f );
+
+				bool bFoundTriangle = false;
+				float CenterWeight = 0.0f;
+				FVector BestEdgeVertex0Position, BestEdgeVertex1Position;
+				{
+					static TArray<FVertexID> PerimeterVertexIDs;
+					EditableMesh->GetPolygonPerimeterVertices( PolygonRef, /* Out */ PerimeterVertexIDs );
+
+					for( int32 VertexNumber = 0; VertexNumber < PerimeterVertexIDs.Num(); ++VertexNumber )
+					{
+						const int32 NextVertexNumber = ( VertexNumber + 1 ) % PerimeterVertexIDs.Num();
+
+
+						const FVertexID EdgeVertex0 = PerimeterVertexIDs[ VertexNumber ];
+						const FVertexID EdgeVertex1 = PerimeterVertexIDs[ NextVertexNumber ];
+
+						const FVector EdgeVertex0Position = EditableMesh->GetVertexAttribute( EdgeVertex0, UEditableMeshAttribute::VertexPosition(), 0 );
+						const FVector EdgeVertex1Position = EditableMesh->GetVertexAttribute( EdgeVertex1, UEditableMeshAttribute::VertexPosition(), 0 );
+
+						BestEdgeVertex0Position = EdgeVertex0Position;
+						BestEdgeVertex1Position = EdgeVertex1Position;
+
+						const FVector VertexWeights = FMath::ComputeBaryCentric2D( ComponentSpaceRayIntersectionWithPolygonPlane, EdgeVertex0Position, EdgeVertex1Position, PolygonCenter );
+						if( VertexWeights.X >= 0.0f && VertexWeights.Y >= 0.0f && VertexWeights.Z >= 0.0f )
+						{
+							CenterWeight = VertexWeights.Z;
+
+							bFoundTriangle = true;
+							break;
+						}
+					}
+				}
+
+				if( bFoundTriangle )
+				{
+					// @todo mesheditor debug
+					// DrawDebugLine( This->GetWorld(), ComponentToWorldMatrix.TransformPosition( BestEdgeVertex0Position ), ComponentToWorldMatrix.TransformPosition( BestEdgeVertex1Position ), FColor::Yellow, false, 0.0f, 0, 2.0f );
+
+					OutInsetProgressTowardCenter = CenterWeight;
+				}
+			}
+		}
+	};
+
+
+	check( !IsUndoSystemAvailable() || GEditor->IsTransactionActive() );
+
+	static TSet<UEditableMesh*> MeshesBeingModified;
+	MeshesBeingModified.Reset();
+
+
+	// NOTE: These are stored in the reverse order they need to be applied in!
+	static TArray< TTuple< UObject*, TUniquePtr<FChange>> > UpdateRevertChanges;
+	UpdateRevertChanges.Reset();
+
+	TSet<UEditableMesh*>* MeshesBeingModifiedPtr = &MeshesBeingModified;
+	TArray< TTuple< UObject*, TUniquePtr<FChange>>>* UpdateRevertChangesPtr = &UpdateRevertChanges;
+
+	auto TrackUndo = [MeshesBeingModifiedPtr, UpdateRevertChangesPtr]( UObject* Object, TUniquePtr<FChange> RevertChange )
+	{
+		if( RevertChange.IsValid() )
+		{
+			check( Object != nullptr );
+
+			// If the object is a mesh, make sure we've started to modify it
+			UEditableMesh* EditableMesh = Cast<UEditableMesh>( Object );
+			if( EditableMesh != nullptr )
+			{
+				// StartModification() must have already been called, otherwise it's too late at this point -- the mesh has been changed
+				// while the render thread was still using it.  Bad things will happen.  So we assert here.
+				check( MeshesBeingModifiedPtr->Contains( EditableMesh ) );
+			}
+
+			UpdateRevertChangesPtr->Add( TTuple<UObject*, TUniquePtr<FChange>>( Object, MoveTemp( RevertChange ) ) );
+		}
+	};
+
+	// Make sure StartModification() is called on all selected meshes
+	const EMeshModificationType MeshModificationType = bIsActionFinishing ? EMeshModificationType::Final : ( bIsFirstActiveActionUpdate ? EMeshModificationType::FirstInterim : EMeshModificationType::Interim );
+	const EMeshTopologyChange MeshTopologyChange = ( ( ActiveAction == EMeshEditAction::Move || ActiveAction == EMeshEditAction::MoveUsingGizmo ) ? EMeshTopologyChange::NoTopologyChange : EMeshTopologyChange::TopologyChange );
+	{
+		for( const FMeshElement& SelectedMeshElement : SelectedMeshElements )
+		{
+			if( SelectedMeshElement.IsValidMeshElement() )
+			{
+				UEditableMesh* EditableMesh = FindOrCreateEditableMesh( *SelectedMeshElement.Component, SelectedMeshElement.ElementAddress.SubMeshAddress );
+				if( EditableMesh != nullptr )
+				{
+					// @todo mesheditor debug
+					//GWarn->Logf( TEXT( "Selected: %s, Element: %s" ), *SelectedMeshElement.Component->GetName(), *SelectedMeshElement.ToString() );
+
+					if( !MeshesBeingModified.Contains( EditableMesh ) )	// @todo gizmo: All transformables will also hit this right?
+					{
+						MeshesBeingModified.Add( EditableMesh );
+
+						EditableMesh->StartModification( MeshModificationType, MeshTopologyChange );
+
+						TrackUndo( EditableMesh, EditableMesh->MakeUndo() );
+					}
+				}
+			}
+		}
+	}
+
+	bool bIsMovingSelectedMeshElements = false;
+
+	if( ActiveAction == EMeshEditAction::Move ||
+		ActiveAction == EMeshEditAction::MoveUsingGizmo )
+	{
+		bIsMovingSelectedMeshElements = true;
+	}
+	else if( ActiveAction == EMeshEditAction::SplitEdge )
+	{
+		// Split edge
+		static TMap< UEditableMesh*, TArray< TTuple< FMeshElement, FEdgeID > > > SelectedMeshesAndEdges;
+		GetSelectedMeshesAndEdges( /* Out */ SelectedMeshesAndEdges );
+
+		if( SelectedMeshesAndEdges.Num() > 0 )
+		{
+			for( auto& MeshAndEdges : SelectedMeshesAndEdges )
+			{
+				UEditableMesh* EditableMesh = MeshAndEdges.Key;
+				verify( !EditableMesh->AnyChangesToUndo() );
+
+				const TArray<TTuple<FMeshElement, FEdgeID>>& Edges = MeshAndEdges.Value;
+
+				// Figure out where to split
+				static TArray<float> Splits;
+				FindEdgeSplitUnderInteractor( ActiveActionInteractor, EditableMesh, Edges, /* Out */ Splits );
+
+				// Split the edges
+				if( Splits.Num() > 0 )
+				{
+					// @todo mesheditor edgeloop: Test code
+					if( false )
+					{
+						if( Splits[ 0 ] > 0.25f )
+						{
+							Splits.Insert( FMath::Max( Splits[ 0 ] - 0.2f, 0.0f ), 0 );
+						}
+						if( Splits.Last() < 0.75f )
+						{
+							Splits.Add( FMath::Min( Splits.Last() + 0.2f, 1.0f ) );
+						}
+					}
+
+					for( const TTuple<FMeshElement, FEdgeID>& EdgeElementAndID : Edges )
+					{
+						const FMeshElement& EdgeMeshElement = EdgeElementAndID.Get<0>();
+						const FEdgeID EdgeID = EdgeElementAndID.Get<1>();
+
+						static TArray<FVertexID> NewVertexIDs;
+						NewVertexIDs.Reset();
+
+						EditableMesh->SplitEdge( EdgeID, Splits, /* Out */ NewVertexIDs );
+
+						// Select all of the new vertices that were created by splitting the edge
+						{
+							bShouldDeselectAllFirst = true;	// Don't keep the edge selected
+							for( const FVertexID NewVertexID : NewVertexIDs )
+							{
+								FMeshElement MeshElementToSelect;
+								{
+									MeshElementToSelect.Component = EdgeMeshElement.Component;
+									MeshElementToSelect.ElementAddress.SubMeshAddress = EdgeMeshElement.ElementAddress.SubMeshAddress;
+									MeshElementToSelect.ElementAddress.ElementType = EEditableMeshElementType::Vertex;
+									MeshElementToSelect.ElementAddress.ElementID = NewVertexID;
+								}
+
+								// Queue selection of this new element.  We don't want it to be part of the current action.
+								MeshElementsToSelect.Add( MeshElementToSelect );
+							}
+						}
+					}
+
+					TrackUndo( EditableMesh, EditableMesh->MakeUndo() );
+				}
+			}
+		}
+	}
+
+	else if( ActiveAction == EMeshEditAction::SplitEdgeAndDragVertex )
+	{
+		if( SplitEdgeMeshesAndEdgesToSplit.Num() > 0 && SplitEdgeSplitList.Num() > 0 )
+		{
+			for( auto& MeshAndEdges : SplitEdgeMeshesAndEdgesToSplit )
+			{
+				UEditableMesh* EditableMesh = MeshAndEdges.Key;
+				verify( !EditableMesh->AnyChangesToUndo() );
+
+				const TArray<TTuple<FMeshElement, FEdgeID>>& Edges = MeshAndEdges.Value;
+
+				for( const TTuple<FMeshElement, FEdgeID>& EdgeElementAndID : Edges )
+				{
+					const FMeshElement& EdgeMeshElement = EdgeElementAndID.Get<0>();
+					const FEdgeID EdgeID = EdgeElementAndID.Get<1>();
+
+					static TArray<FVertexID> NewVertexIDs;
+					NewVertexIDs.Reset();
+
+					EditableMesh->SplitEdge( EdgeID, SplitEdgeSplitList, /* Out */ NewVertexIDs );
+
+					// Select all of the new vertices that were created by splitting the edge
+					{
+						bShouldDeselectAllFirst = true;	// Don't keep the edge selected
+						for( const FVertexID NewVertexID : NewVertexIDs )
+						{
+							FMeshElement MeshElementToSelect;
+							{
+								MeshElementToSelect.Component = EdgeMeshElement.Component;
+								MeshElementToSelect.ElementAddress.SubMeshAddress = EdgeMeshElement.ElementAddress.SubMeshAddress;
+								MeshElementToSelect.ElementAddress.ElementType = EEditableMeshElementType::Vertex;
+								MeshElementToSelect.ElementAddress.ElementID = NewVertexID;
+							}
+
+							// Queue selection of this new element.  We don't want it to be part of the current action.
+							MeshElementsToSelect.Add( MeshElementToSelect );
+
+						}
+
+						if( NewVertexIDs.Num() > 0 )
+						{
+							// Drag the split elements right away!
+							bIsMovingSelectedMeshElements = true;
+						}
+					}
+				}
+
+				TrackUndo( EditableMesh, EditableMesh->MakeUndo() );
+			}
+		}
+	}
+
+	else if( ActiveAction == EMeshEditAction::InsertEdgeLoop )
+	{
+		// Insert edge loop
+		static TMap< UEditableMesh*, TArray< TTuple< FMeshElement, FEdgeID > > > SelectedMeshesAndEdges;
+		GetSelectedMeshesAndEdges( /* Out */ SelectedMeshesAndEdges );
+
+		if( SelectedMeshesAndEdges.Num() > 0 )
+		{
+			// Deselect the edges first, since they'll be deleted or split up while inserting the edge loop,
+			// and we want them to be re-selected after undo
+			{
+				FCompoundChangeInput CompoundRevertInput;
+
+				FSelectOrDeselectMeshElementsChangeInput ChangeInput;
+				for( auto& MeshAndEdges : SelectedMeshesAndEdges )
+				{
+					UEditableMesh* EditableMesh = MeshAndEdges.Key;
+					const TArray<TTuple<FMeshElement, FEdgeID>>& Edges = MeshAndEdges.Value;
+
+					for( const TTuple<FMeshElement, FEdgeID>& EdgeElementAndID : Edges )
+					{
+						const FMeshElement& EdgeMeshElement = EdgeElementAndID.Get<0>();
+						ChangeInput.MeshElementsToDeselect.Add( EdgeMeshElement );
+					}
+				}
+
+				CompoundRevertInput.Subchanges.Add( MoveTemp( FSelectOrDeselectMeshElementsChange( MoveTemp( ChangeInput ) ).Execute( MeshEditorModeProxyObject ) ) );
+
+				TrackUndo( MeshEditorModeProxyObject, MakeUnique<FCompoundChange>( MoveTemp( CompoundRevertInput ) ) );
+			}
+
+
+			for( auto& MeshAndEdges : SelectedMeshesAndEdges )
+			{
+				UEditableMesh* EditableMesh = MeshAndEdges.Key;
+				verify( !EditableMesh->AnyChangesToUndo() );
+
+				const TArray<TTuple<FMeshElement, FEdgeID>>& Edges = MeshAndEdges.Value;
+
+				// Figure out where to add the loop along the edge
+				static TArray<float> Splits;
+				FindEdgeSplitUnderInteractor( ActiveActionInteractor, EditableMesh, Edges, /* Out */ Splits );
+
+				// Insert the edge loop
+				if( Splits.Num() > 0 )
+				{
+					// @todo mesheditor edgeloop: Test code
+					if( false )
+					{
+						if( Splits[ 0 ] > 0.25f )
+						{
+							Splits.Insert( FMath::Max( Splits[ 0 ] - 0.2f, 0.0f ), 0 );
+						}
+						if( Splits.Last() < 0.75f )
+						{
+							Splits.Add( FMath::Min( Splits.Last() + 0.2f, 1.0f ) );
+						}
+					}
+
+					for( const TTuple<FMeshElement, FEdgeID>& EdgeElementAndID : Edges )
+					{
+						const FMeshElement& EdgeMeshElement = EdgeElementAndID.Get<0>();
+						const FEdgeID EdgeID = EdgeElementAndID.Get<1>();
+
+						static TArray<FEdgeID> NewEdgeIDs;
+						NewEdgeIDs.Reset();
+
+						EditableMesh->InsertEdgeLoop( EdgeID, Splits, /* Out */ NewEdgeIDs );
+
+						// Select all of the new edges that were created by inserting the loop
+						{
+							bShouldDeselectAllFirst = true;	// Don't keep the edge selected
+							for( const FEdgeID NewEdgeID : NewEdgeIDs )
+							{
+								FMeshElement MeshElementToSelect;
+								{
+									MeshElementToSelect.Component = EdgeMeshElement.Component;
+									MeshElementToSelect.ElementAddress.SubMeshAddress = EdgeMeshElement.ElementAddress.SubMeshAddress;
+									MeshElementToSelect.ElementAddress.ElementType = EEditableMeshElementType::Edge;
+									MeshElementToSelect.ElementAddress.ElementID = NewEdgeID;
+								}
+
+								// Queue selection of this new element.  We don't want it to be part of the current action.
+								MeshElementsToSelect.Add( MeshElementToSelect );
+							}
+						}
+					}
+
+					TrackUndo( EditableMesh, EditableMesh->MakeUndo() );
+				}
+			}
+		}
+	}
+
+	else if( ActiveAction == EMeshEditAction::ExtrudePolygon ||
+			 ActiveAction == EMeshEditAction::FreelyExtrudePolygon )
+	{
+		// Extrude polygon
+		static TMap< UEditableMesh*, TArray< TTuple< FMeshElement, FPolygonRef > > > MeshesWithPolygonsToExtrude;
+		GetSelectedMeshesAndPolygons( /* Out */ MeshesWithPolygonsToExtrude );
+
+		if( MeshesWithPolygonsToExtrude.Num() > 0 )
+		{
+			if( ActiveAction == EMeshEditAction::FreelyExtrudePolygon )
+			{
+				// Start moving the extruded polygon right away!
+				bIsMovingSelectedMeshElements = true;
+			}
+
+			// Deselect the mesh elements before we delete them.  This will make sure they become selected again after undo.
+			{
+				FCompoundChangeInput CompoundRevertInput;
+
+				FSelectOrDeselectMeshElementsChangeInput ChangeInput;
+				for( auto& MeshAndPolygons : MeshesWithPolygonsToExtrude )
+				{
+					UEditableMesh* EditableMesh = MeshAndPolygons.Key;
+					const TArray<TTuple<FMeshElement, FPolygonRef>>& PolygonsToExtrude = MeshAndPolygons.Value;
+
+					for( const TTuple<FMeshElement, FPolygonRef>& PolygonToExtrude : PolygonsToExtrude )
+					{
+						const FMeshElement& PolygonElement = PolygonToExtrude.Get<0>();
+						ChangeInput.MeshElementsToDeselect.Add( PolygonElement );
+					}
+				}
+
+				CompoundRevertInput.Subchanges.Add( MoveTemp( FSelectOrDeselectMeshElementsChange( MoveTemp( ChangeInput ) ).Execute( MeshEditorModeProxyObject ) ) );
+
+				TrackUndo( MeshEditorModeProxyObject, MakeUnique<FCompoundChange>( MoveTemp( CompoundRevertInput ) ) );
+			}
+
+			for( auto& MeshAndPolygons : MeshesWithPolygonsToExtrude )
+			{
+				UEditableMesh* EditableMesh = MeshAndPolygons.Key;
+				const TArray<TTuple<FMeshElement, FPolygonRef>>& PolygonsToExtrude = MeshAndPolygons.Value;
+
+				static TArray<FPolygonRef> PolygonRefsToExtrude;
+				PolygonRefsToExtrude.Reset();
+				for( const TTuple<FMeshElement, FPolygonRef>& PolygonToExtrude : PolygonsToExtrude )
+				{
+					PolygonRefsToExtrude.Add( PolygonToExtrude.Get<1>() );
+				}
+
+
+				float ExtrudeDistance = 0.0f;
+
+				if( ActiveAction == EMeshEditAction::ExtrudePolygon )
+				{
+					// Figure out how far the extruded polygon should be from where it started
+					const float AxisLength = 10000.0f;	// @todo mesheditor tweak (ideally should be infinite)
+					FindExtrudeDistanceUsingInteractor(
+						this,
+						ActiveActionInteractor,
+						EditableMesh,
+						ExtrudePolygonAxisOrigin,
+						ExtrudePolygonAxisDirection,
+						AxisLength,
+						/* Out */ ExtrudeDistance );
+				}
+
+
+				{
+					verify( !EditableMesh->AnyChangesToUndo() );
+
+					// Position the new polygon to where the interactor is
+					UPrimitiveComponent* Component = PolygonsToExtrude[ 0 ].Get<0>().Component.Get();	// NOTE: All polygons in this array belong to the same mesh/component, so we just need the first element
+					check( Component != nullptr );
+					// @todo mesheditor: We're working with a float here, so we'll treat the component scale as a scalar (X)
+					const float ComponentSpaceExtrudeDistance = ExtrudeDistance / Component->ComponentToWorld.GetScale3D().X;
+
+					// Create a copy of the polygon with new extruded polygons for each edge
+					static TArray<FPolygonRef> NewExtrudedFrontPolygons;
+					NewExtrudedFrontPolygons.Reset();
+					EditableMesh->ExtrudePolygons( PolygonRefsToExtrude, ComponentSpaceExtrudeDistance, /* Out */ NewExtrudedFrontPolygons );
+
+
+					// Make sure the new polygons are selected.  The old polygon was deleted and will become deselected automatically.
+					for( int32 NewPolygonNumber = 0; NewPolygonNumber < NewExtrudedFrontPolygons.Num(); ++NewPolygonNumber )
+					{
+						const FPolygonRef& NewExtrudedFrontPolygon = NewExtrudedFrontPolygons[ NewPolygonNumber ];
+
+						const FMeshElement& MeshElement = PolygonsToExtrude[ NewPolygonNumber ].Get<0>();
+
+						FMeshElement NewExtrudedPolygonMeshElement;
+						{
+							NewExtrudedPolygonMeshElement.Component = MeshElement.Component;
+							NewExtrudedPolygonMeshElement.ElementAddress = MeshElement.ElementAddress;
+							NewExtrudedPolygonMeshElement.ElementAddress.SectionID = NewExtrudedFrontPolygon.SectionID;
+							NewExtrudedPolygonMeshElement.ElementAddress.ElementID = NewExtrudedFrontPolygon.PolygonID;
+						}
+
+						// Queue selection of this new element.  We don't want it to be part of the current action.
+						MeshElementsToSelect.Add( NewExtrudedPolygonMeshElement );
+					}
+				}
+
+				TrackUndo( EditableMesh, EditableMesh->MakeUndo() );
+			}
+		}
+	}
+	else if( ActiveAction == EMeshEditAction::InsetPolygon ||
+			 ActiveAction == EMeshEditAction::BevelPolygon )
+	{
+		static TMap< UEditableMesh*, TArray< TTuple< FMeshElement, FPolygonRef > > > MeshesWithPolygonsToInset;
+		GetSelectedMeshesAndPolygons( /* Out */ MeshesWithPolygonsToInset );
+
+		if( MeshesWithPolygonsToInset.Num() > 0 )
+		{
+			if( InsetUsingPolygonElement.IsValidMeshElement() &&
+				InsetUsingPolygonElement.ElementAddress.ElementType == EEditableMeshElementType::Polygon &&
+				GetSelectedMeshElementIndex( InsetUsingPolygonElement ) != INDEX_NONE )
+			{
+				UPrimitiveComponent* InsetUsingComponent = InsetUsingPolygonElement.Component.Get();
+				check( InsetUsingComponent != nullptr );
+
+				UEditableMesh* InsetUsingEditableMesh = FindOrCreateEditableMesh( *InsetUsingComponent, InsetUsingPolygonElement.ElementAddress.SubMeshAddress );
+				if( InsetUsingEditableMesh != nullptr )
+				{
+					const FPolygonRef InsetUsingPolygonRef( InsetUsingPolygonElement.ElementAddress.SectionID, FPolygonID( InsetUsingPolygonElement.ElementAddress.ElementID ) );
+
+					// Figure out how far to inset the polygon
+					float InsetFixedDistance = 0.0f;
+					float InsetProgressTowardCenter = 0.0f;
+					FindInsetAmount(
+						this,
+						ActiveActionInteractor,
+						InsetUsingPolygonRef,
+						*InsetUsingComponent,
+						InsetUsingEditableMesh,
+						/* Out */ InsetFixedDistance,
+						/* Out */ InsetProgressTowardCenter );
+
+					if( InsetFixedDistance > SMALL_NUMBER ||
+						InsetProgressTowardCenter > SMALL_NUMBER )
+					{
+
+						// Deselect the mesh elements before we delete them.  This will make sure they become selected again after undo.
+						{
+							FCompoundChangeInput CompoundRevertInput;
+
+							FSelectOrDeselectMeshElementsChangeInput ChangeInput;
+							for( auto& MeshAndPolygons : MeshesWithPolygonsToInset )
+							{
+								UEditableMesh* EditableMesh = MeshAndPolygons.Key;
+								const TArray<TTuple<FMeshElement, FPolygonRef>>& PolygonsToDeselect = MeshAndPolygons.Value;
+
+								for( const TTuple<FMeshElement, FPolygonRef>& PolygonToDeselect : PolygonsToDeselect )
+								{
+									const FMeshElement& PolygonElement = PolygonToDeselect.Get<0>();
+									ChangeInput.MeshElementsToDeselect.Add( PolygonElement );
+								}
+							}
+
+							CompoundRevertInput.Subchanges.Add( MoveTemp( FSelectOrDeselectMeshElementsChange( MoveTemp( ChangeInput ) ).Execute( MeshEditorModeProxyObject ) ) );
+
+							TrackUndo( MeshEditorModeProxyObject, MakeUnique<FCompoundChange>( MoveTemp( CompoundRevertInput ) ) );
+						}
+
+						for( auto& MeshAndPolygons : MeshesWithPolygonsToInset )
+						{
+							UEditableMesh* EditableMesh = MeshAndPolygons.Key;
+							const TArray<TTuple<FMeshElement, FPolygonRef>>& PolygonsToInset = MeshAndPolygons.Value;
+
+							UPrimitiveComponent* Component = PolygonsToInset[ 0 ].Get<0>().Component.Get();	// NOTE: All polygons in this array belong to the same mesh/component, so we just need the first element
+							check( Component != nullptr );
+
+
+							static TArray<FPolygonRef> PolygonRefsToInset;
+							PolygonRefsToInset.Reset();
+							for( const TTuple<FMeshElement, FPolygonRef>& PolygonToInset : PolygonsToInset )
+							{
+								PolygonRefsToInset.Add( PolygonToInset.Get<1>() );
+							}
+
+							{
+								verify( !EditableMesh->AnyChangesToUndo() );
+
+								// Inset time!!
+								static TArray<FPolygonRef> NewCenterInsetPolygons;
+								static TArray<FPolygonRef> NewSideInsetPolygons;
+								if( ActiveAction == EMeshEditAction::InsetPolygon )
+								{
+									const EInsetPolygonsMode InsetPolygonsMode = EInsetPolygonsMode::All;	// @todo mesheditor inset: Make configurable?
+																											// @todo mesheditor inset: Add options for Fixed distance (instead of Percentage distance, like now.)
+
+									EditableMesh->InsetPolygons( PolygonRefsToInset, InsetFixedDistance, InsetProgressTowardCenter, InsetPolygonsMode, /* Out */ NewCenterInsetPolygons, /* Out */ NewSideInsetPolygons );
+								}
+								else if( ensure( ActiveAction == EMeshEditAction::BevelPolygon ) )
+								{
+									EditableMesh->BevelPolygons( PolygonRefsToInset, InsetFixedDistance, InsetProgressTowardCenter, /* Out */ NewCenterInsetPolygons, /* Out */ NewSideInsetPolygons );
+								}
+
+								// Make sure the new polygons are selected.  The old polygon was deleted and will become deselected automatically.
+								if( NewCenterInsetPolygons.Num() > 0 )
+								{
+									for( int32 NewPolygonNumber = 0; NewPolygonNumber < NewCenterInsetPolygons.Num(); ++NewPolygonNumber )
+									{
+										const FPolygonRef& NewInsetCenterPolygon = NewCenterInsetPolygons[ NewPolygonNumber ];
+
+										const FMeshElement& MeshElement = PolygonsToInset[ NewPolygonNumber ].Get<0>();
+
+										FMeshElement PolygonMeshElement;
+										{
+											PolygonMeshElement.Component = MeshElement.Component;
+											PolygonMeshElement.ElementAddress = MeshElement.ElementAddress;
+											PolygonMeshElement.ElementAddress.SectionID = NewInsetCenterPolygon.SectionID;
+											PolygonMeshElement.ElementAddress.ElementID = NewInsetCenterPolygon.PolygonID;
+										}
+
+										// Queue selection of this new element.  We don't want it to be part of the current action.
+										MeshElementsToSelect.Add( PolygonMeshElement );
+									}
+								}
+								else if( ensure( NewSideInsetPolygons.Num() > 0 ) )
+								{
+									for( int32 NewPolygonNumber = 0; NewPolygonNumber < NewSideInsetPolygons.Num(); ++NewPolygonNumber )
+									{
+										const FPolygonRef& NewInsetSidePolygon = NewSideInsetPolygons[ NewPolygonNumber ];
+
+										const FMeshElement& MeshElement = PolygonsToInset[ NewPolygonNumber ].Get<0>();
+
+										FMeshElement PolygonMeshElement;
+										{
+											PolygonMeshElement.Component = MeshElement.Component;
+											PolygonMeshElement.ElementAddress = MeshElement.ElementAddress;
+											PolygonMeshElement.ElementAddress.SectionID = NewInsetSidePolygon.SectionID;
+											PolygonMeshElement.ElementAddress.ElementID = NewInsetSidePolygon.PolygonID;
+										}
+
+										// Queue selection of this new element.  We don't want it to be part of the current action.
+										MeshElementsToSelect.Add( PolygonMeshElement );
+									}
+								}
+							}
+
+							TrackUndo( EditableMesh, EditableMesh->MakeUndo() );
+						}
+					}
+				}
+			}
+		}
+	}
+	else if( ActiveAction == EMeshEditAction::ExtendEdge )
+	{
+		// Extend edge
+		static TMap< UEditableMesh*, TArray< TTuple< FMeshElement, FEdgeID > > > MeshesWithEdgesToExtend;
+		GetSelectedMeshesAndEdges( /* Out */ MeshesWithEdgesToExtend );
+
+		if( MeshesWithEdgesToExtend.Num() > 0 )
+		{
+			// Allow the user to move the new extended edge around right away!
+			bIsMovingSelectedMeshElements = true;
+
+			for( auto& MeshAndEdges : MeshesWithEdgesToExtend )
+			{
+				UEditableMesh* EditableMesh = MeshAndEdges.Key;
+				const TArray<TTuple<FMeshElement, FEdgeID>>& EdgeElementsToExtend = MeshAndEdges.Value;
+
+				static TArray<FEdgeID> EdgeIDsToExtend;
+				EdgeIDsToExtend.Reset();
+
+				for( const TTuple<FMeshElement, FEdgeID>& EdgeElementToExtend : EdgeElementsToExtend )
+				{
+					const FMeshElement& PolygonElement = EdgeElementToExtend.Get<0>();
+
+					EdgeIDsToExtend.Add( EdgeElementToExtend.Get<1>() );
+				}
+
+
+				verify( !EditableMesh->AnyChangesToUndo() );
+
+				{
+					// Extend the edge
+					const bool bWeldNeighbors = true;	// @todo mesheditor urgent extrude: Make optional somehow
+					static TArray<FEdgeID> NewExtendedEdgeIDs;
+					NewExtendedEdgeIDs.Reset();
+					EditableMesh->ExtendEdges( EdgeIDsToExtend, bWeldNeighbors, /* Out */ NewExtendedEdgeIDs );
+
+
+					// Make sure the new edges are selected
+					bShouldDeselectAllFirst = true;	// Deselect the original edges
+					for( int32 NewEdgeNumber = 0; NewEdgeNumber < NewExtendedEdgeIDs.Num(); ++NewEdgeNumber )
+					{
+						const FEdgeID NewExtendedEdgeID = NewExtendedEdgeIDs[ NewEdgeNumber ];
+
+						const FMeshElement& MeshElement = EdgeElementsToExtend[ NewEdgeNumber ].Get<0>();
+
+						FMeshElement NewExtendedEdgeMeshElement;
+						{
+							NewExtendedEdgeMeshElement.Component = MeshElement.Component;
+							NewExtendedEdgeMeshElement.ElementAddress = MeshElement.ElementAddress;
+							NewExtendedEdgeMeshElement.ElementAddress.ElementID = NewExtendedEdgeID;
+						}
+
+						// Queue selection of this new element.  We don't want it to be part of the current action.
+						MeshElementsToSelect.Add( NewExtendedEdgeMeshElement );
+					}
+				}
+
+				TrackUndo( EditableMesh, EditableMesh->MakeUndo() );
+			}
+		}
+	}
+
+	else if( ActiveAction == EMeshEditAction::ExtendVertex )
+	{
+		static TMap< UEditableMesh*, TArray< TTuple< FMeshElement, FVertexID > > > MeshesWithVerticesToExtend;
+		GetSelectedMeshesAndVertices( /* Out */ MeshesWithVerticesToExtend );
+
+		if( MeshesWithVerticesToExtend.Num() > 0 )
+		{
+			// Allow the user to move the new extended vertex around right away!
+			bIsMovingSelectedMeshElements = true;
+
+			for( auto& MeshAndVertices : MeshesWithVerticesToExtend )
+			{
+				UEditableMesh* EditableMesh = MeshAndVertices.Key;
+				const TArray<TTuple<FMeshElement, FVertexID>>& VertexElementsToExtend = MeshAndVertices.Value;
+
+				static TArray<FVertexID> VertexIDsToExtend;
+				VertexIDsToExtend.Reset();
+
+				for( const TTuple<FMeshElement, FVertexID>& VertexElementToExtend : VertexElementsToExtend )
+				{
+					const FMeshElement& PolygonElement = VertexElementToExtend.Get<0>();
+
+					VertexIDsToExtend.Add( VertexElementToExtend.Get<1>() );
+				}
+
+
+				verify( !EditableMesh->AnyChangesToUndo() );
+
+				{
+					UPrimitiveComponent* ComponentPtr = VertexElementsToExtend[ 0 ].Get<0>().Component.Get();
+					check( ComponentPtr != nullptr );
+					UPrimitiveComponent& Component = *ComponentPtr;
+					const FMatrix ComponentToWorldMatrix = Component.GetRenderMatrix();
+
+					const FVector ComponentSpaceDragToLocation = ComponentToWorldMatrix.InverseTransformPosition( ActiveActionInteractor->GetInteractorData().LastDragToLocation );
+
+					// Extend the vertex
+					const bool bOnlyExtendClosestEdge = true;	// @todo mesheditor urgent extrude: Make optional somehow
+					static TArray<FVertexID> NewExtendedVertexIDs;
+					NewExtendedVertexIDs.Reset();
+					EditableMesh->ExtendVertices( VertexIDsToExtend, bOnlyExtendClosestEdge, ComponentSpaceDragToLocation, /* Out */ NewExtendedVertexIDs );
+
+
+					// Make sure the new vertices are selected
+					bShouldDeselectAllFirst = true;	// Deselect the original vertices
+					for( int32 NewVertexNumber = 0; NewVertexNumber < NewExtendedVertexIDs.Num(); ++NewVertexNumber )
+					{
+						const FVertexID NewExtendedVertexID = NewExtendedVertexIDs[ NewVertexNumber ];
+
+						const FMeshElement& MeshElement = VertexElementsToExtend[ NewVertexNumber ].Get<0>();
+
+						FMeshElement NewExtendedVertexMeshElement;
+						{
+							NewExtendedVertexMeshElement.Component = MeshElement.Component;
+							NewExtendedVertexMeshElement.ElementAddress = MeshElement.ElementAddress;
+							NewExtendedVertexMeshElement.ElementAddress.ElementID = NewExtendedVertexID;
+						}
+
+						// Queue selection of this new element.  We don't want it to be part of the current action.
+						MeshElementsToSelect.Add( NewExtendedVertexMeshElement );
+					}
+				}
+
+				TrackUndo( EditableMesh, EditableMesh->MakeUndo() );
+			}
+		}
+	}
+
+	else if( ActiveAction == EMeshEditAction::EditVertexCornerSharpness ||
+			 ActiveAction == EMeshEditAction::EditEdgeCreaseSharpness )
+	{
+		static TMap< UEditableMesh*, TArray< TTuple< FMeshElement, FVertexID > > > MeshesWithSelectedVertices;
+		static TMap< UEditableMesh*, TArray< TTuple< FMeshElement, FEdgeID > > > MeshesWithSelectedEdges;
+		GetSelectedMeshesAndVertices( /* Out */ MeshesWithSelectedVertices );
+		GetSelectedMeshesAndEdges( /* Out */ MeshesWithSelectedEdges );
+
+		if( MeshesWithSelectedVertices.Num() > 0 || MeshesWithSelectedEdges.Num() > 0 )
+		{
+			// @todo mesheditor subdiv: Hard coded tweakables; ideally should be sized in screen space
+			const float DragScaleFactor = 5.0f;
+			const float ProgressBarHeight = 1000.0f;
+
+			// Figure out how much to either increase or decrease sharpness based on how far the user has dragged up or down in world space.
+			float DragDeltaZ = 0.0f;
+			
+			{
+				FMeshEditorInteractorData& MeshEditorInteractorData = GetMeshEditorInteractorData( ActiveActionInteractor );
+				if( MeshEditorInteractorData.bLaserIsValid || MeshEditorInteractorData.bGrabberSphereIsValid )
+				{
+					const FVector ProgressBarStart( EditSharpnessStartLocation + FVector( 0.0f, 0.0f, -ProgressBarHeight * 0.5f ) );
+					const FVector ProgressBarEnd( EditSharpnessStartLocation + FVector( 0.0f, 0.0f, ProgressBarHeight * 0.5f ) );
+
+					// First check to see if we're within grabber sphere range.  Otherwise, we'll use the laser.
+					bool bIsInRange = false;
+					FVector ClosestPointOnProgressBar;
+					if( MeshEditorInteractorData.bGrabberSphereIsValid )
+					{
+						if( FMath::PointDistToSegment( MeshEditorInteractorData.GrabberSphere.Center, ProgressBarStart, ProgressBarEnd ) <= MeshEditorInteractorData.GrabberSphere.W )
+						{
+							bIsInRange = true;
+							ClosestPointOnProgressBar = FMath::ClosestPointOnSegment(
+								MeshEditorInteractorData.GrabberSphere.Center,
+								ProgressBarStart, ProgressBarEnd );
+						}
+					}
+					
+					if( !bIsInRange && ensure( MeshEditorInteractorData.bLaserIsValid ) )
+					{
+						bIsInRange = true;
+						FVector ClosestPointOnRay;
+						FMath::SegmentDistToSegment(
+							ProgressBarStart, ProgressBarEnd,
+							MeshEditorInteractorData.LaserStart, MeshEditorInteractorData.LaserEnd,
+							/* Out */ ClosestPointOnProgressBar,
+							/* Out */ ClosestPointOnRay );
+					}
+
+					if( bIsInRange )
+					{
+						// Generate a drag value between -1.0 and 1.0 based on the interaction of the ray and the progress bar line segment.
+						const float ProgressBarLength = ( ProgressBarEnd - ProgressBarStart ).Size();
+						DragDeltaZ = -1.0f + 2.0f * ( ( ClosestPointOnProgressBar - ProgressBarStart ).Size() / ProgressBarLength );
+					}
+				}
+			}
+
+			const float ScaledDragDelta = DragDeltaZ * DragScaleFactor;
+
+			for( auto& MeshAndSelectedVertices : MeshesWithSelectedVertices )
+			{
+				UEditableMesh* EditableMesh = MeshAndSelectedVertices.Key;
+				const TArray<TTuple<FMeshElement, FVertexID>>& VertexElementsAndIDs = MeshAndSelectedVertices.Value;
+
+				static TArray<FVertexID> VertexIDs;
+				VertexIDs.Reset();
+
+				static TArray<float> NewSharpnessValues;
+				NewSharpnessValues.Reset();
+
+				for( const TTuple<FMeshElement, FVertexID>& VertexElementAndID : VertexElementsAndIDs )
+				{
+					const FVertexID VertexID = VertexElementAndID.Get<1>();
+
+					const float CurrentSharpnessValue = EditableMesh->GetVertexAttribute( VertexID, UEditableMeshAttribute::VertexCornerSharpness(), 0 ).X;
+					const float NewSharpnessValue = FMath::Clamp( CurrentSharpnessValue + ScaledDragDelta, 0.0f, 1.0f );
+
+					VertexIDs.Add( VertexID );
+					NewSharpnessValues.Add( NewSharpnessValue );
+				}
+
+
+				verify( !EditableMesh->AnyChangesToUndo() );
+
+				EditableMesh->SetVerticesCornerSharpness( VertexIDs, NewSharpnessValues );
+
+				TrackUndo( EditableMesh, EditableMesh->MakeUndo() );
+			}
+
+			for( auto& MeshAndSelectedEdges : MeshesWithSelectedEdges )
+			{
+				UEditableMesh* EditableMesh = MeshAndSelectedEdges.Key;
+				const TArray<TTuple<FMeshElement, FEdgeID>>& EdgeElementsAndIDs = MeshAndSelectedEdges.Value;
+
+				static TArray<FEdgeID> EdgeIDs;
+				EdgeIDs.Reset();
+
+				static TArray<float> NewSharpnessValues;
+				NewSharpnessValues.Reset();
+
+				for( const TTuple<FMeshElement, FEdgeID>& EdgeElementAndID : EdgeElementsAndIDs )
+				{
+					const FEdgeID EdgeID = EdgeElementAndID.Get<1>();
+
+					const float CurrentSharpnessValue = EditableMesh->GetEdgeAttribute( EdgeID, UEditableMeshAttribute::EdgeCreaseSharpness(), 0 ).X;
+					const float NewSharpnessValue = FMath::Clamp( CurrentSharpnessValue + ScaledDragDelta, 0.0f, 1.0f );
+
+					EdgeIDs.Add( EdgeID );
+					NewSharpnessValues.Add( NewSharpnessValue );
+				}
+
+				verify( !EditableMesh->AnyChangesToUndo() );
+
+				EditableMesh->SetEdgesCreaseSharpness( EdgeIDs, NewSharpnessValues );
+
+				TrackUndo( EditableMesh, EditableMesh->MakeUndo() );
+			}
+		}
+	}
+	else if( ActiveAction == EMeshEditAction::DrawVertices )
+	{
+		// @todo mesheditor: need a way to determine the plane we are going to create the polygon on.
+		// This is because the depth of placed points is ambiguous in a perspective viewport.
+		// For testing purposes, define a hardcoded plane.
+		const FPlane PolygonPlane( FVector( 0, 0, 50 ), FVector( 0, 0, 1 ) );
+
+		FMeshEditorInteractorData& MeshEditorInteractorData = GetMeshEditorInteractorData( ActiveActionInteractor );
+		if( MeshEditorInteractorData.bLaserIsValid || MeshEditorInteractorData.bGrabberSphereIsValid )
+		{
+			UEditableMesh* EditableMesh = nullptr;
+			UPrimitiveComponent* Component = nullptr;
+			FEditableMeshSubMeshAddress SubMeshAddress;
+
+			if( SelectedMeshElements.Num() == 0 )
+			{
+				// @todo mesheditor: support creating a new mesh from scratch here
+				// Look into support for creating new assets in the transient package without needing to specify a filename?
+			}
+			else
+			{
+				// Currently adds new vertices to whichever editable mesh is currently selected
+				// @todo mesheditor: allow multiple selected meshes? What should this do?
+
+				for( int32 SelectedElementIndex = 0; SelectedElementIndex < SelectedMeshElements.Num(); ++SelectedElementIndex )
+				{
+					const FMeshElement& SelectedMeshElement = SelectedMeshElements[ SelectedElementIndex ];
+					if( SelectedMeshElement.IsValidMeshElement() )
+					{
+						EditableMesh = FindOrCreateEditableMesh( *SelectedMeshElement.Component, SelectedMeshElement.ElementAddress.SubMeshAddress );
+						if( EditableMesh != nullptr )
+						{
+							Component = SelectedMeshElement.Component.Get();
+							SubMeshAddress = SelectedMeshElement.ElementAddress.SubMeshAddress;
+							break;
+						}
+					}
+				}
+
+				check( EditableMesh );
+				check( Component );
+
+				FVector Point = FMath::LinePlaneIntersection( MeshEditorInteractorData.LaserStart, MeshEditorInteractorData.LaserEnd, PolygonPlane );
+
+				// @todo mesheditor: Hard coded tweakables. MinDistanceSqr should probably be in screen space.
+				const float MinDistanceSqr = FMath::Square( 5.0f );
+				const float MinTimeToPlacePoint = 0.25f;
+				const float AngleThreshold = 0.86f;		// cos(30 degrees)
+
+				const double CurrentTime = FSlateApplication::Get().GetCurrentTime();
+
+				if( DrawnPoints.Num() == 0 )
+				{
+					// Always place the first point regardless
+					DrawnPoints.Emplace( CurrentTime, Point );
+				}
+				else if ( DrawnPoints.Num() == 1 )
+				{
+					// Place the second point if it's sufficiently far away from the first
+					if( FVector::DistSquared( Point, DrawnPoints.Last().Get<1>() ) > MinDistanceSqr )
+					{
+						DrawnPoints.Emplace( CurrentTime, Point );
+					}
+				}
+				else
+				{
+					// Function which determines whether the segment formed with the given endpoint intersects with any other segment
+					auto IsSelfIntersecting = []( const FPlane& Plane, const TArray<TTuple<double, FVector>>& Points, const FVector& EndPoint )
+					{
+						// Calculate a 2d basis for the plane (origin and 2d axes)
+						const FVector PlaneNormal( Plane );
+						check( PlaneNormal.IsNormalized() );
+						const FVector PlaneOrigin( PlaneNormal * Plane.W );
+
+						const FVector DirectionX = ( PlaneNormal.X == 0.0f && PlaneNormal.Y == 0.0f )
+							? FVector( PlaneNormal.Z, 0.0f, -PlaneNormal.X )
+							: FVector( -PlaneNormal.Y, PlaneNormal.X, 0.0f ).GetSafeNormal();
+
+						const FVector DirectionY = FVector::CrossProduct( PlaneNormal, DirectionX );
+
+						// Transforms a point in 3D space into the basis on the plane described by an origin and two orthogonal direction vectors on the plane
+						auto ToPlaneBasis = []( const FVector& Origin, const FVector& DirectionX, const FVector& DirectionY, const FVector& Point )
+						{
+							const FVector Offset( Point - Origin );
+							return FVector2D( FVector::DotProduct( Offset, DirectionX ), FVector::DotProduct( Offset, DirectionY ) );
+						};
+
+						// Determine whether two line segments intersect in 2D space
+						// @todo core: Put this into FMath static library?
+						auto DoSegmentsIntersect = []( const FVector2D& Start1, const FVector2D& End1, const FVector2D& Start2, const FVector2D& End2 )
+						{
+							const FVector2D Dir1( End1 - Start1 );
+							const FVector2D Dir2( End2 - Start2 );
+							const FVector2D Offset( Start2 - Start1 );
+							const float Det = FVector2D::CrossProduct( Dir1, Dir2 );
+							if( Det == 0.0f )
+							{
+								// Determinant of zero implies parallel segments.
+								// If the below cross product is also zero, this indicates colinear segments which we'll consider an intersection.
+								return ( FVector2D::CrossProduct( Offset, Dir1 ) == 0.0f );
+							}
+
+							const float OneOverDet = 1.0f / Det;
+							const float Intersect1 = FVector2D::CrossProduct( Offset, Dir2 ) * OneOverDet;
+							const float Intersect2 = FVector2D::CrossProduct( Offset, Dir1 ) * OneOverDet;
+
+							return ( Intersect1 >= 0.0f && Intersect1 <= 1.0f && Intersect2 >= 0.0f && Intersect2 <= 1.0f );
+						};
+
+						const int32 NumPoints = Points.Num();
+						const FVector2D SegmentToTestStart = ToPlaneBasis( PlaneOrigin, DirectionX, DirectionY, Points[ NumPoints - 1 ].Get<1>() );
+						const FVector2D SegmentToTestEnd = ToPlaneBasis( PlaneOrigin, DirectionX, DirectionY, EndPoint );
+
+						for( int32 Index = 0; Index < Points.Num() - 2; Index++ )
+						{
+							FVector2D IntersectionPoint;
+							const FVector2D Start = ToPlaneBasis( PlaneOrigin, DirectionX, DirectionY, Points[ Index ].Get<1>() );
+							const FVector2D End = ToPlaneBasis( PlaneOrigin, DirectionX, DirectionY, Points[ Index + 1 ].Get<1>() );
+							if( DoSegmentsIntersect( Start, End, SegmentToTestStart, SegmentToTestEnd ) )
+							{
+								return true;
+							}
+						}
+
+						return false;
+					};
+
+					// Place subsequent points if:
+					// a) they are sufficiently far away from the previous point; and
+					// b) they do not form a self-intersecting poly; and
+					// c) they make a sufficiently big angle with the previous edge; or
+					// d) there was a small pause in the drawing movement
+					const int32 NumDrawnPoints = DrawnPoints.Num();
+					const FVector Point1 = DrawnPoints[ NumDrawnPoints - 2 ].Get<1>();
+					const FVector Point2 = DrawnPoints[ NumDrawnPoints - 1 ].Get<1>();
+					if( FVector::DistSquared( Point, Point2 ) > MinDistanceSqr &&
+					    !IsSelfIntersecting( PolygonPlane, DrawnPoints, Point ) )
+					{
+						if( CurrentTime - DrawnPoints[ NumDrawnPoints - 1 ].Get<0>() > MinTimeToPlacePoint ||
+						    FVector::DotProduct( ( Point2 - Point1 ).GetSafeNormal(), ( Point - Point2 ).GetSafeNormal() ) < AngleThreshold )
+						{
+							// Point is distinct enough from the last, add a new one
+							DrawnPoints.Emplace( CurrentTime, Point );
+						}
+						else
+						{
+							// Point is an extension of the previous edge, update the previous point
+							DrawnPoints[ NumDrawnPoints - 1 ] = MakeTuple( CurrentTime, Point );
+						}
+					}
+				}
+
+				// Create new vertices
+				static TArray<FVertexID> NewVertexIDs;
+				static TArray<FVertexToCreate> VerticesToCreate;
+				NewVertexIDs.Reset( DrawnPoints.Num() );
+				VerticesToCreate.Reset( DrawnPoints.Num() );
+
+				for( const TTuple<double, FVector>& DrawnPoint : DrawnPoints )
+				{
+					VerticesToCreate.Emplace();
+					FVertexToCreate& VertexToCreate = VerticesToCreate.Last();
+
+					VertexToCreate.VertexAttributes.Attributes.Emplace(
+						UEditableMeshAttribute::VertexPosition(),
+						0,
+						FVector4( Component->GetComponentTransform().InverseTransformPosition( DrawnPoint.Get<1>() ) )
+					);
+				}
+
+				EditableMesh->CreateVertices( VerticesToCreate, NewVertexIDs );
+
+				// Select new vertices
+				bShouldDeselectAllFirst = true;
+				for( FVertexID VertexID : NewVertexIDs )
+				{
+					MeshElementsToSelect.Emplace(
+						Component,
+						SubMeshAddress,
+						VertexID
+					);
+				}
+
+				if( DrawnPoints.Num() == 2 )
+				{
+					// If only two points, create an edge
+					static TArray<FEdgeID> NewEdgeIDs;
+					static TArray<FEdgeToCreate> EdgesToCreate;
+					NewEdgeIDs.Reset( 1 );
+					EdgesToCreate.Reset( 1 );
+
+					EdgesToCreate.Emplace();
+					FEdgeToCreate& EdgeToCreate = EdgesToCreate.Last();
+					EdgeToCreate.VertexID0 = NewVertexIDs[ 0 ];
+					EdgeToCreate.VertexID1 = NewVertexIDs[ 1 ];
+
+					EdgeToCreate.EdgeAttributes.Attributes.Emplace( UEditableMeshAttribute::EdgeIsHard(), 0, FVector4( 1.0f ) );
+
+					EditableMesh->CreateEdges( EdgesToCreate, NewEdgeIDs );
+				}
+				else if( DrawnPoints.Num() > 2 )
+				{
+					// If more than two points, create a polygon
+					static TArray<FPolygonRef> NewPolygonRefs;
+					static TArray<FEdgeID> NewEdgeIDs;
+					static TArray<FPolygonToCreate> PolygonsToCreate;
+					NewPolygonRefs.Reset( 1 );
+					NewEdgeIDs.Reset( 1 );
+					PolygonsToCreate.Reset( 1 );
+
+					// Find first valid section to add the polygon to
+					FSectionID SectionID = FSectionID::Invalid;
+					const int32 SectionArraySize = EditableMesh->GetSectionArraySize();
+					for( int32 SectionIndex = 0; SectionIndex < SectionArraySize; ++SectionIndex )
+					{
+						if( EditableMesh->IsValidSection( FSectionID( SectionIndex ) ) )
+						{
+							SectionID = FSectionID( SectionIndex );
+							break;
+						}
+					}
+					check( SectionID != FSectionID::Invalid );
+
+					// Create new polygon
+					PolygonsToCreate.Emplace();
+					FPolygonToCreate& PolygonToCreate = PolygonsToCreate.Last();
+					PolygonToCreate.SectionID = SectionID;
+
+					for( FVertexID NewVertexID : NewVertexIDs )
+					{
+						PolygonToCreate.PerimeterVertices.Emplace();
+						FVertexAndAttributes& VertexAndAttributes = PolygonToCreate.PerimeterVertices.Last();
+						VertexAndAttributes.VertexID = NewVertexID;
+					}
+
+					EditableMesh->CreatePolygons( PolygonsToCreate, NewPolygonRefs, NewEdgeIDs );
+
+					// Check if the polygon normal is pointing towards us. If not, we need to flip the polygon
+					FVector PolygonNormal = EditableMesh->ComputePolygonNormal( NewPolygonRefs[ 0 ] );
+
+					if( CachedCameraToWorld.IsSet() )
+					{
+						if( FVector::DotProduct( Component->GetComponentTransform().TransformVector( PolygonNormal ), DrawnPoints[ 0 ].Get<1>() - CachedCameraToWorld.GetValue().GetLocation() ) > 0.0f )
+						{
+							EditableMesh->FlipPolygons( NewPolygonRefs );
+							PolygonNormal = -PolygonNormal;
+						}
+					}
+
+					// Set polygon vertex normals (assuming hard edges)
+					static TArray<FVertexAttributesForPolygon> VertexAttributesForPolygon;
+					VertexAttributesForPolygon.Reset( 1 );
+
+					VertexAttributesForPolygon.Emplace();
+					FVertexAttributesForPolygon& VertexAttrs = VertexAttributesForPolygon.Last();
+					VertexAttrs.PolygonRef = NewPolygonRefs[ 0 ];
+
+					for( int32 PolygonVertexIndex = 0; PolygonVertexIndex < NewVertexIDs.Num(); ++PolygonVertexIndex )
+					{
+						VertexAttrs.PerimeterVertexAttributeLists.Emplace();
+						FMeshElementAttributeList& AttributeList = VertexAttrs.PerimeterVertexAttributeLists.Last();
+
+						AttributeList.Attributes.Emplace(
+							UEditableMeshAttribute::VertexNormal(),
+							0,
+							FVector4( PolygonNormal )
+						);
+					}
+				}
+
+				TrackUndo( EditableMesh, EditableMesh->MakeUndo() );
+			}
+		}
+	}
+
+	// Update selection state.  Note that we intentionally do this BEFORE we perform any dragging, so that
+	// we'll be dragging any newly-generated geometry from the mesh edit action.  For example, when extending
+	// an edge we want to drag around the newly-created edge, not the edge that was selected before.
+	{
+		FCompoundChangeInput CompoundRevertInput;
+
+		if( bShouldDeselectAllFirst )
+		{
+			CompoundRevertInput.Subchanges.Add( MoveTemp( FDeselectAllMeshElementsChange( FDeselectAllMeshElementsChangeInput() ).Execute( MeshEditorModeProxyObject ) ) );
+		}
+
+		// Make sure nothing is still selected that was deleted
+		{
+			TUniquePtr<FChange> ClearInvalidRevertChange = MoveTemp( ClearInvalidSelectedElements_Internal() );
+			if( ClearInvalidRevertChange.IsValid() )
+			{
+				CompoundRevertInput.Subchanges.Add( MoveTemp( ClearInvalidRevertChange ) );
+			}
+		}
+
+		if( MeshElementsToSelect.Num() > 0 )
+		{
+			FSelectOrDeselectMeshElementsChangeInput ChangeInput;
+			ChangeInput.MeshElementsToSelect = MeshElementsToSelect;
+			CompoundRevertInput.Subchanges.Add( MoveTemp( FSelectOrDeselectMeshElementsChange( MoveTemp( ChangeInput ) ).Execute( MeshEditorModeProxyObject ) ) );
+		}
+
+		if( CompoundRevertInput.Subchanges.Num() > 0 )
+		{
+			TrackUndo( MeshEditorModeProxyObject, MakeUnique<FCompoundChange>( MoveTemp( CompoundRevertInput ) ) );
+		}
+	}
+
+	if( bIsMovingSelectedMeshElements )
+	{
+		static TMap< UEditableMesh*, TArray< const FMeshElementViewportTransformable* > > MeshesAndTransformables;
+		MeshesAndTransformables.Reset();
+
+		{
+			const TArray<TUniquePtr<class FViewportTransformable>>& Transformables = ViewportWorldInteraction->GetTransformables();
+			for( const TUniquePtr<class FViewportTransformable>& TransformablePtr : Transformables )
+			{
+				const FViewportTransformable& Transformable = *TransformablePtr;
+
+				// @todo gizmo: Can we only bother updating elements that actually have moved? (LastTransform isn't useful here because of tick order)
+//					if( !Transformable.LastTransform.Equals( Transformable.GetTransform() ) )
+				{
+					const FMeshElementViewportTransformable& MeshElementTransformable = static_cast<const FMeshElementViewportTransformable&>( Transformable );
+					const FMeshElement& ElementToMove = MeshElementTransformable.MeshElement;
+
+					UPrimitiveComponent* ComponentPtr = ElementToMove.Component.Get();
+					check( ComponentPtr != nullptr );
+					UPrimitiveComponent& Component = *ComponentPtr;
+
+					UEditableMesh* EditableMesh = FindOrCreateEditableMesh( Component, ElementToMove.ElementAddress.SubMeshAddress );
+					check( EditableMesh != nullptr );
+
+					MeshesAndTransformables.FindOrAdd( EditableMesh ).Add( &MeshElementTransformable );
+				}
+			}
+		}
+
+
+		for( auto& MeshAndTransformables : MeshesAndTransformables )
+		{
+			static TArray<FVertexToMove> VerticesToMove;
+			VerticesToMove.Reset();
+
+			// We use a TSet, so that the same vertex (from the same mesh) isn't moved more than once
+			static TSet<FVertexID> VertexIDsAlreadyMoved;
+			VertexIDsAlreadyMoved.Reset();
+
+			UEditableMesh* EditableMesh = MeshAndTransformables.Key;
+
+			const TArray< const FMeshElementViewportTransformable* >& TransformablesForMesh = MeshAndTransformables.Value;
+
+			UPrimitiveComponent* ComponentPtr = TransformablesForMesh[0]->MeshElement.Component.Get();
+			check( ComponentPtr != nullptr );
+			UPrimitiveComponent& Component = *ComponentPtr;
+
+			const FTransform ComponentToWorld = Component.GetComponentToWorld();
+			const FTransform WorldToComponent = ComponentToWorld.Inverse();
+
+			for( const FMeshElementViewportTransformable* TransformablePtr : TransformablesForMesh )
+			{
+				const FMeshElementViewportTransformable& MeshElementTransformable = *TransformablePtr;
+				const FMeshElement& ElementToMove = MeshElementTransformable.MeshElement;
+				check( ElementToMove.IsValidMeshElement() );
+
+				// Build a matrix that transforms any vertex (in component space) using the mesh elements current
+				// transform (in world space), and finally back to a final position in component space
+				const FTransform ComponentDeltaFromStartTransform =
+					ComponentToWorld *
+					MeshElementTransformable.StartTransform.Inverse() *
+					MeshElementTransformable.CurrentTransform *
+					WorldToComponent;
+
+				if( ElementToMove.ElementAddress.ElementType == EEditableMeshElementType::Vertex )
+				{
+					const FVertexID VertexID( ElementToMove.ElementAddress.ElementID );
+
+					if( !VertexIDsAlreadyMoved.Contains( VertexID ) )
+					{
+						const FVector NewVertexPosition = ComponentToWorld.InverseTransformPosition( MeshElementTransformable.CurrentTransform.GetLocation() );
+
+						FVertexToMove& VertexToMove = *new( VerticesToMove ) FVertexToMove();
+						VertexToMove.VertexID = VertexID;
+						VertexToMove.NewVertexPosition = NewVertexPosition;
+
+						VertexIDsAlreadyMoved.Add( VertexToMove.VertexID );
+					}
+				}
+				else if( ElementToMove.ElementAddress.ElementType == EEditableMeshElementType::Edge )
+				{
+					const FEdgeID EdgeID( ElementToMove.ElementAddress.ElementID );
+
+					FVertexID EdgeVertexIDs[ 2 ];
+					EditableMesh->GetEdgeVertices( EdgeID, /* Out */ EdgeVertexIDs[ 0 ], /* Out */ EdgeVertexIDs[ 1 ] );
+
+					for( const FVertexID EdgeVertexID : EdgeVertexIDs )
+					{
+						if( !VertexIDsAlreadyMoved.Contains( EdgeVertexID ) )
+						{
+							const FVector OriginalComponentSpaceVertexPosition = EditableMesh->GetVertexAttribute( EdgeVertexID, UEditableMeshAttribute::VertexPosition(), 0 );
+							const FVector NewComponentSpaceVertexPosition = ComponentDeltaFromStartTransform.TransformPosition( OriginalComponentSpaceVertexPosition );
+
+							FVertexToMove& VertexToMove = *new( VerticesToMove ) FVertexToMove();
+							VertexToMove.VertexID = EdgeVertexID;
+							VertexToMove.NewVertexPosition = NewComponentSpaceVertexPosition;
+
+							VertexIDsAlreadyMoved.Add( VertexToMove.VertexID );
+						}
+					}
+				}
+				else if( ElementToMove.ElementAddress.ElementType == EEditableMeshElementType::Polygon )
+				{
+					const FPolygonRef PolygonRef( ElementToMove.ElementAddress.SectionID, FPolygonID( ElementToMove.ElementAddress.ElementID ) );
+
+					static TArray<FVertexID> PolygonPerimeterVertexIDs;
+					EditableMesh->GetPolygonPerimeterVertices( PolygonRef, /* Out */ PolygonPerimeterVertexIDs );
+
+					for( const FVertexID PolygonPerimeterVertexID : PolygonPerimeterVertexIDs )
+					{
+						if( !VertexIDsAlreadyMoved.Contains( PolygonPerimeterVertexID ) )
+						{
+							const FVector OriginalComponentSpaceVertexPosition = EditableMesh->GetVertexAttribute( PolygonPerimeterVertexID, UEditableMeshAttribute::VertexPosition(), 0 );
+							const FVector NewComponentSpaceVertexPosition = ComponentDeltaFromStartTransform.TransformPosition( OriginalComponentSpaceVertexPosition );
+
+							FVertexToMove& VertexToMove = *new( VerticesToMove ) FVertexToMove();
+							VertexToMove.VertexID = PolygonPerimeterVertexID;
+							VertexToMove.NewVertexPosition = NewComponentSpaceVertexPosition;
+
+							VertexIDsAlreadyMoved.Add( VertexToMove.VertexID );
+						}
+					}
+				}
+			}
+
+			if( VerticesToMove.Num() > 0 )
+			{
+				verify( !EditableMesh->AnyChangesToUndo() );
+
+				// @todo mesheditor perf: Polygons will actually be triangulated twice in the same frame when performing an operation that implicitly retriangulates at the end of the op (once before move, once after move.)  Not super ideal.
+				EditableMesh->MoveVertices( VerticesToMove );
+
+				TrackUndo( EditableMesh, EditableMesh->MakeUndo() );
+			}
+		}
+	}
+
+
+	// Make sure EndModification() is called on any EditableMesh objects that were modified, so their graphics and physics
+	// state is updated.
+	{
+		for( UEditableMesh* EditableMesh : MeshesBeingModified )
+		{
+			if( EditableMesh != nullptr )
+			{
+				verify( !EditableMesh->AnyChangesToUndo() );
+				EditableMesh->EndModification();
+				TrackUndo( EditableMesh, EditableMesh->MakeUndo() );
+			}
+		}
+	}
+	
+
+	// Store final undo state for all changes that happened in this update
+	{
+		if( MeshModificationType == EMeshModificationType::Final )
+		{
+			// If we're finalizing the action, this will save the undo state for everything that happened in this function, including
+			// selection changes.  
+			for( TTuple<UObject*, TUniquePtr<FChange>>& UpdateRevertChange : UpdateRevertChanges )
+			{
+				UObject* Object = UpdateRevertChange.Get<0>();
+				check( Object != nullptr );
+
+				TUniquePtr<FChange>& RevertChange = UpdateRevertChange.Get<1>();
+				check( RevertChange.IsValid() );
+
+				StoreUndo( Object, MoveTemp( RevertChange ) );
+			}
+		}
+		else
+		{
+			// Otherwise, we'll store the commands to undo in our 'PreviewRevertChanges' member, so they can be
+			// rolled back at the beginning of the next frame before any new interactions take place.  This allows the user to preview
+			// (potentially highly destructive) changes live!
+			// NOTE: UpdateRevertChanges has entires stored in reverse order, so we'll iterate backwards as we build up or undo buffer
+			for( int32 ChangeIndex = UpdateRevertChanges.Num() - 1; ChangeIndex >= 0; --ChangeIndex )
+			{
+				TTuple<UObject*, TUniquePtr<FChange>>& UpdateRevertChange = UpdateRevertChanges[ ChangeIndex ];
+
+				PreviewRevertChanges.Add( MoveTemp( UpdateRevertChange ) );
+			}
+
+			PreviewTopologyChange = MeshTopologyChange;
+		}
+	}
+
+	// Reset temporary containers
+	UpdateRevertChanges.Reset();
+	MeshesBeingModified.Reset();
+
+	bIsFirstActiveActionUpdate = false;
+}
+
+
+void FMeshEditorMode::GetSelectedMeshesAndElements( EEditableMeshElementType ElementType, TMap<UEditableMesh*, TArray<FMeshElement>>& OutMeshesAndElements )
+{
+	OutMeshesAndElements.Reset();
+
+	for( int32 SelectedElementIndex = 0; SelectedElementIndex < this->SelectedMeshElements.Num(); ++SelectedElementIndex )
+	{
+		const FMeshElement& SelectedMeshElement = this->SelectedMeshElements[ SelectedElementIndex ];
+		if( SelectedMeshElement.IsValidMeshElement() )
+		{
+			UEditableMesh* EditableMesh = this->FindOrCreateEditableMesh( *SelectedMeshElement.Component, SelectedMeshElement.ElementAddress.SubMeshAddress );
+			if( EditableMesh != nullptr )
+			{
+				if( ElementType == EEditableMeshElementType::Any || SelectedMeshElement.ElementAddress.ElementType == ElementType )
+				{
+					TArray<FMeshElement>& Elements = OutMeshesAndElements.FindOrAdd( EditableMesh );
+					Elements.Add( SelectedMeshElement );
+				}
+			}
+		}
+	}
+}
+
+
+void FMeshEditorMode::GetSelectedMeshesAndVertices( TMap< UEditableMesh*, TArray< TTuple< FMeshElement, FVertexID > > >& OutMeshesAndVertices )
+{
+	OutMeshesAndVertices.Reset();
+
+	for( int32 SelectedElementIndex = 0; SelectedElementIndex < this->SelectedMeshElements.Num(); ++SelectedElementIndex )
+	{
+		const FMeshElement& SelectedMeshElement = this->SelectedMeshElements[ SelectedElementIndex ];
+		if( SelectedMeshElement.IsValidMeshElement() )
+		{
+			UEditableMesh* EditableMesh = this->FindOrCreateEditableMesh( *SelectedMeshElement.Component, SelectedMeshElement.ElementAddress.SubMeshAddress );
+			if( EditableMesh != nullptr )
+			{
+				if( SelectedMeshElement.ElementAddress.ElementType == EEditableMeshElementType::Vertex )
+				{
+					TArray< TTuple< FMeshElement, FVertexID > >& Vertices = OutMeshesAndVertices.FindOrAdd( EditableMesh );
+					Vertices.Add( TTuple< FMeshElement, FVertexID >( SelectedMeshElement, SelectedMeshElement.ElementAddress.ElementID ) );
+				}
+			}
+		}
+	}
+}
+
+
+void FMeshEditorMode::GetSelectedMeshesAndEdges( TMap< UEditableMesh*, TArray< TTuple< FMeshElement, FEdgeID > > >& OutMeshesAndEdges )
+{
+	OutMeshesAndEdges.Reset();
+
+	for( int32 SelectedElementIndex = 0; SelectedElementIndex < this->SelectedMeshElements.Num(); ++SelectedElementIndex )
+	{
+		const FMeshElement& SelectedMeshElement = this->SelectedMeshElements[ SelectedElementIndex ];
+		if( SelectedMeshElement.IsValidMeshElement() )
+		{
+			UEditableMesh* EditableMesh = this->FindOrCreateEditableMesh( *SelectedMeshElement.Component, SelectedMeshElement.ElementAddress.SubMeshAddress );
+			if( EditableMesh != nullptr )
+			{
+				if( SelectedMeshElement.ElementAddress.ElementType == EEditableMeshElementType::Edge )
+				{
+					TArray< TTuple< FMeshElement, FEdgeID > >& Edges = OutMeshesAndEdges.FindOrAdd( EditableMesh );
+					Edges.Add( TTuple< FMeshElement, FEdgeID >( SelectedMeshElement, SelectedMeshElement.ElementAddress.ElementID ) );
+				}
+			}
+		}
+	}
+}
+
+
+void FMeshEditorMode::GetSelectedMeshesAndPolygons(	TMap< UEditableMesh*, TArray< TTuple< FMeshElement, FPolygonRef > > >& OutMeshesAndPolygons )
+{
+	OutMeshesAndPolygons.Reset();
+
+	for( int32 SelectedElementIndex = 0; SelectedElementIndex < this->SelectedMeshElements.Num(); ++SelectedElementIndex )
+	{
+		const FMeshElement& SelectedMeshElement = this->SelectedMeshElements[ SelectedElementIndex ];
+		if( SelectedMeshElement.IsValidMeshElement() )
+		{
+			UEditableMesh* EditableMesh = this->FindOrCreateEditableMesh( *SelectedMeshElement.Component, SelectedMeshElement.ElementAddress.SubMeshAddress );
+			if( EditableMesh != nullptr )
+			{
+				if( SelectedMeshElement.ElementAddress.ElementType == EEditableMeshElementType::Polygon )
+				{
+					TArray< TTuple< FMeshElement, FPolygonRef > >& Polygons = OutMeshesAndPolygons.FindOrAdd( EditableMesh );
+					Polygons.Add( TTuple< FMeshElement, FPolygonRef >( SelectedMeshElement, FPolygonRef( SelectedMeshElement.ElementAddress.SectionID, FPolygonID( SelectedMeshElement.ElementAddress.ElementID ) ) ) );
+				}
+			}
+		}
+	}
+};
+
+
+void FMeshEditorMode::FindEdgeSplitUnderInteractor(	UViewportInteractor* ViewportInteractor, const UEditableMesh* EditableMesh, const TArray<TTuple<FMeshElement, FEdgeID>>& Edges, TArray<float>& OutSplits )
+{
+	check( ViewportInteractor != nullptr );
+
+	OutSplits.Reset();
+
+	// Figure out where to split based on where the interactor is aiming.  We'll look at all of the
+	// selected edges, and choose a split offset based on the closest point along one of those edges
+	// to the interactor.  All selected edges will then be split by the same proportion.
+	float ClosestDistanceToEdge = MAX_flt;
+	for( const TTuple<FMeshElement, FEdgeID>& EdgeElementAndID : Edges )
+	{
+		const FMeshElement& EdgeMeshElement = EdgeElementAndID.Get<0>();
+		const FEdgeID EdgeID = EdgeElementAndID.Get<1>();
+
+		FMeshEditorInteractorData& MeshEditorInteractorData = this->GetMeshEditorInteractorData( ViewportInteractor );
+		if( MeshEditorInteractorData.bLaserIsValid || MeshEditorInteractorData.bGrabberSphereIsValid )
+		{
+			FVertexID EdgeVertexIDs[ 2 ];
+			EditableMesh->GetEdgeVertices( EdgeID, /* Out */ EdgeVertexIDs[0], /* Out */ EdgeVertexIDs[1] );
+
+			const FTransform ComponentToWorld = EdgeMeshElement.Component.Get()->GetComponentToWorld();
+
+			FVector WorldSpaceEdgeVertexPositions[ 2 ];
+			for( int32 EdgeVertexNumber = 0; EdgeVertexNumber < 2; ++EdgeVertexNumber )
+			{
+				WorldSpaceEdgeVertexPositions[ EdgeVertexNumber ] =
+					ComponentToWorld.TransformPosition( EditableMesh->GetVertexAttribute( EdgeVertexIDs[ EdgeVertexNumber ], UEditableMeshAttribute::VertexPosition(), 0 ) );
+			}
+
+
+			// Compute how far along the edge the interactor is aiming
+			// @todo mesheditor: HoverLocation is only valid when actually hovering over some mesh element.  Really we probably want
+			// to just use the impact point of whatever is under the interactor, even if it's not an editable mesh.
+			const FVector WorldSpaceClosestPointOnEdge = FMath::ClosestPointOnSegment( MeshEditorInteractorData.HoverLocation, WorldSpaceEdgeVertexPositions[ 0 ], WorldSpaceEdgeVertexPositions[ 1 ] );
+
+			// How close are we to this edge?
+			const float DistanceToEdge = ( MeshEditorInteractorData.HoverLocation - WorldSpaceClosestPointOnEdge ).Size();
+			if( DistanceToEdge <= ClosestDistanceToEdge )
+			{
+				ClosestDistanceToEdge = DistanceToEdge;
+
+				const float WorldSpaceEdgeLength = ( WorldSpaceEdgeVertexPositions[ 1 ] - WorldSpaceEdgeVertexPositions[ 0 ] ).Size();
+				float ProgressAlongEdge = 0.0f;
+				if( WorldSpaceEdgeLength > 0.0f )
+				{
+					// NOTE: This should never actually need to be clamped, but we do it just to avoid floating point precision problems
+					// where the value is slightly smaller than zero or greater than one
+
+					// @todo mesheditor: Splitting an edge at position 0 or 1 will introduce a coincident point and degenerate polygons.  Might want to
+					// have a practical min and max progress amount?
+					ProgressAlongEdge = FMath::Clamp(
+						( WorldSpaceClosestPointOnEdge - WorldSpaceEdgeVertexPositions[ 0 ] ).Size() / WorldSpaceEdgeLength,
+						0.0f,
+						1.0f );
+				}
+
+				OutSplits.Reset();
+				OutSplits.Add( ProgressAlongEdge );
+			}
+		}
+	}
+}
+
+
+FEditableMeshElementAddress FMeshEditorMode::QueryElement( const UEditableMesh& EditableMesh, const bool bUseSphere, const FSphere& Sphere, const float SphereFuzzyDistance, bool bUseRay, const FVector& RayStart, const FVector& RayEnd, const float RayFuzzyDistance, const EEditableMeshElementType OnlyElementType, const FVector& CameraLocation, const float FuzzyDistanceScaleFactor, EInteractorShape& OutInteractorShape, FVector& OutHitLocation ) const
+{
+	OutHitLocation = FVector::ZeroVector;
+
+	FEditableMeshElementAddress HitElementAddress;
+	HitElementAddress.SubMeshAddress = EditableMesh.GetSubMeshAddress();
+
+	static TSet<FVertexID> FrontFacingVertices;
+	FrontFacingVertices.Reset();
+	
+	static TSet<FEdgeID> FrontFacingEdges;
+	FrontFacingEdges.Reset();
+
+	static TSet<FPolygonRef> FrontFacingPolygons;
+	FrontFacingPolygons.Reset();
+
+	// Look for all the front-facing elements
+	for( int32 SectionIndex = 0; SectionIndex < EditableMesh.GetSectionArraySize(); ++SectionIndex )
+	{
+		const FSectionID SectionID( SectionIndex );
+		if( EditableMesh.IsValidSection( SectionID ) )
+		{
+			const uint32 SectionPolygonArraySize = EditableMesh.GetPolygonArraySize( SectionID );
+			for( uint32 SectionPolygonIndex = 0; SectionPolygonIndex < SectionPolygonArraySize; ++SectionPolygonIndex )
+			{
+				const FPolygonID PolygonID( SectionPolygonIndex );
+
+				const FPolygonRef PolygonRef( SectionID, PolygonID );
+				if( EditableMesh.IsValidPolygon( PolygonRef ) )
+				{
+					const FVector PolygonNormal = EditableMesh.ComputePolygonNormal( PolygonRef );
+					const FVector PolygonCenter = EditableMesh.ComputePolygonCenter( PolygonRef );
+					if( FVector::DotProduct( CameraLocation - PolygonCenter, PolygonNormal ) > 0.0f )
+					{
+						FrontFacingPolygons.Add( PolygonRef );
+
+						const int32 PolygonVertexCount = EditableMesh.GetPolygonPerimeterVertexCount( PolygonRef );
+						for( int32 Index = 0; Index < PolygonVertexCount; ++Index )
+						{
+							FrontFacingVertices.Add( EditableMesh.GetPolygonPerimeterVertex( PolygonRef, Index ) );
+							bool bOutEdgeWindingIsReversedForPolygons;
+							FrontFacingEdges.Add( EditableMesh.GetPolygonPerimeterEdge( PolygonRef, Index, bOutEdgeWindingIsReversedForPolygons ) );
+						}
+					}
+				}
+			}
+		}
+	}
+
+	EInteractorShape ClosestInteractorShape = EInteractorShape::Invalid;
+	FVector ClosestHitLocation = FVector::ZeroVector;
+	float ClosestDistanceOnRay = TNumericLimits<float>::Max();
+	float ClosestDistanceToRay = TNumericLimits<float>::Max();
+	FVector CurrentRayEnd = RayEnd;
+
+	// Check polygons first; this is so we always impose a closest hit location at the poly before checking other elements, so anything behind is occluded
+	for( const FPolygonRef PolygonRef : FrontFacingPolygons )
+	{
+		static TArray<FVertexID> MeshVertexIDs;
+		MeshVertexIDs.Reset();
+		EditableMesh.GetPolygonPerimeterVertices( PolygonRef, /* Out */ MeshVertexIDs );
+
+		const uint32 PolygonTriangleCount = EditableMesh.GetPolygonTriangulatedTriangleCount( PolygonRef );
+		for( uint32 PolygonTriangleNumber = 0; PolygonTriangleNumber < PolygonTriangleCount; ++PolygonTriangleNumber )
+		{
+			FVector TriangleVertexPositions[ 3 ];
+			for( uint32 TriangleVertexNumber = 0; TriangleVertexNumber < 3; ++TriangleVertexNumber )
+			{
+				TriangleVertexPositions[ TriangleVertexNumber ] =
+					EditableMesh.GetPolygonTriangulatedTriangleVertexPosition( PolygonRef, PolygonTriangleNumber, TriangleVertexNumber );
+			}
+
+			const bool bAlreadyHitTriangle = ( HitElementAddress.ElementType == EEditableMeshElementType::Polygon );
+			const bool bHit = CheckTriangle( bUseSphere, Sphere, SphereFuzzyDistance, bUseRay, RayStart, CurrentRayEnd, RayFuzzyDistance, TriangleVertexPositions, CameraLocation, FuzzyDistanceScaleFactor, ClosestInteractorShape, ClosestDistanceToRay, ClosestDistanceOnRay, ClosestHitLocation, bAlreadyHitTriangle );
+			if( bHit )
+			{
+				HitElementAddress.ElementType = EEditableMeshElementType::Polygon;
+				HitElementAddress.SectionID = PolygonRef.SectionID;
+				HitElementAddress.ElementID = PolygonRef.PolygonID;
+			}
+		}
+	}
+
+	// Reset the closest distance to ray (which will have been set to 0 by the polygon check) so other elements can be found within the fuzzy distance
+	ClosestDistanceToRay = TNumericLimits<float>::Max();
+
+	// Check edges
+	if( OnlyElementType == EEditableMeshElementType::Invalid || OnlyElementType == EEditableMeshElementType::Edge )
+	{
+		for( const FEdgeID EdgeID : FrontFacingEdges )
+		{
+			FVector EdgeVertexPositions[ 2 ];
+			EdgeVertexPositions[ 0 ] = EditableMesh.GetVertexAttribute( EditableMesh.GetEdgeVertex( EdgeID, 0 ), UEditableMeshAttribute::VertexPosition(), 0 );
+			EdgeVertexPositions[ 1 ] = EditableMesh.GetVertexAttribute( EditableMesh.GetEdgeVertex( EdgeID, 1 ), UEditableMeshAttribute::VertexPosition(), 0 );
+
+			const bool bAlreadyHitEdge = ( HitElementAddress.ElementType == EEditableMeshElementType::Edge );
+			const bool bHit = CheckEdge( bUseSphere, Sphere, SphereFuzzyDistance, bUseRay, RayStart, CurrentRayEnd, RayFuzzyDistance, EdgeVertexPositions, CameraLocation, FuzzyDistanceScaleFactor, ClosestInteractorShape, ClosestDistanceToRay, ClosestDistanceOnRay, ClosestHitLocation, bAlreadyHitEdge );
+			if( bHit )
+			{
+				HitElementAddress.ElementType = EEditableMeshElementType::Edge;
+				HitElementAddress.SectionID = FSectionID::Invalid;
+				HitElementAddress.ElementID = EdgeID;
+			}
+		}
+	}
+
+	// @todo mesheditor perf:  OMG too slow.  Needs acceleration structure + MT.  Horribly slow to iterate over all elements like this.  We need a proper spatial query!
+
+	ClosestDistanceToRay = TNumericLimits<float>::Max();
+
+	// Check vertices
+	if( OnlyElementType == EEditableMeshElementType::Invalid || OnlyElementType == EEditableMeshElementType::Vertex )
+	{
+		for( const FVertexID VertexID : FrontFacingVertices )
+		{
+			const FVector VertexPosition = EditableMesh.GetVertexAttribute( VertexID, UEditableMeshAttribute::VertexPosition(), 0 );
+			const bool bAlreadyHitVertex = ( HitElementAddress.ElementType == EEditableMeshElementType::Vertex );
+			const bool bHit = CheckVertex( bUseSphere, Sphere, SphereFuzzyDistance, bUseRay, RayStart, CurrentRayEnd, RayFuzzyDistance, VertexPosition, CameraLocation, FuzzyDistanceScaleFactor, ClosestInteractorShape, ClosestDistanceToRay, ClosestDistanceOnRay, ClosestHitLocation, bAlreadyHitVertex );
+			if( bHit )
+			{
+				HitElementAddress.ElementType = EEditableMeshElementType::Vertex;
+				HitElementAddress.SectionID = FSectionID::Invalid;
+				HitElementAddress.ElementID = VertexID;
+			}
+		}
+	}
+
+	if( HitElementAddress.ElementType != EEditableMeshElementType::Invalid )
+	{
+		OutInteractorShape = ClosestInteractorShape;
+		OutHitLocation = ClosestHitLocation;
+	}
+
+	return HitElementAddress;
+}
+
+
+bool FMeshEditorMode::CheckVertex( const bool bUseSphere, const FSphere& Sphere, const float SphereFuzzyDistance, bool bUseRay, const FVector& RayStart, const FVector& RayEnd, const float RayFuzzyDistance, const FVector& VertexPosition, const FVector& CameraLocation, const float FuzzyDistanceScaleFactor, EInteractorShape& ClosestInteractorShape, float& ClosestDistanceToRay, float& ClosestDistanceOnRay, FVector& ClosestHitLocation, const bool bAlreadyHitVertex )
+{
+	bool bHit = false;
+
+ 	const float DistanceToCamera = ( CameraLocation - VertexPosition ).Size();
+	const float DistanceBasedScaling = MeshEd::OverlayDistanceBias->GetFloat() + DistanceToCamera * FuzzyDistanceScaleFactor;
+	check( DistanceBasedScaling > 0.0f );
+
+	if( bUseSphere )
+	{
+		const float DistanceToSphere = ( VertexPosition - Sphere.Center ).Size();
+		if( DistanceToSphere <= Sphere.W )
+		{
+			// Inside sphere
+			if( DistanceToSphere < ClosestDistanceToRay ||
+				( !bAlreadyHitVertex && FMath::Abs( DistanceToSphere - ClosestDistanceToRay ) < SphereFuzzyDistance * DistanceBasedScaling ) )
+			{
+				ClosestDistanceToRay = DistanceToSphere;
+				ClosestDistanceOnRay = 0.0f;
+				ClosestHitLocation = VertexPosition;
+				ClosestInteractorShape = EInteractorShape::GrabberSphere;
+
+				bHit = true;
+			}
+		}
+	}
+
+	if( bUseRay )
+	{
+		const FVector ClosestPointOnRay = FMath::ClosestPointOnSegment( VertexPosition, RayStart, RayEnd );
+		const float DistanceToRay = ( ClosestPointOnRay - VertexPosition ).Size();
+		const float DistanceOnRay = ( ClosestPointOnRay - RayStart ).Size();
+
+		const FVector RayDirection = ( RayEnd - RayStart ).GetSafeNormal();
+		const FVector DirectionTowardClosestPointOnRay = ( ClosestPointOnRay - RayStart ).GetSafeNormal();
+		const bool bIsBehindRay = FVector::DotProduct( RayDirection, DirectionTowardClosestPointOnRay ) < 0.0f;
+		if( !bIsBehindRay )
+		{
+			// Are we within the minimum distance for hitting a vertex?
+			if( DistanceToRay < RayFuzzyDistance * DistanceBasedScaling )
+			{
+				const bool bWithinFuzzyRadius = FMath::Abs( DistanceOnRay - ClosestDistanceOnRay ) < RayFuzzyDistance * DistanceBasedScaling;
+
+				if( ( bWithinFuzzyRadius && DistanceToRay < ClosestDistanceToRay ) || ( !bWithinFuzzyRadius && DistanceOnRay < ClosestDistanceOnRay ) )
+				{
+					ClosestDistanceToRay = DistanceToRay;
+					ClosestDistanceOnRay = DistanceOnRay;
+					ClosestHitLocation = ClosestPointOnRay;
+					ClosestInteractorShape = EInteractorShape::Laser;
+					bHit = true;
+				}
+
+				// @todo mesheditor debug
+				// const float Radius = GHackComponentToWorld.InverseTransformVector( FVector( RayFuzzyDistance * DistanceBasedScaling, 0.0f, 0.0f ) ).X;
+				// DrawDebugSphere( GHackVWI->GetWorld(), GHackComponentToWorld.TransformPosition( VertexPosition ), Radius, 8, bHit ? FColor::Green : FColor::Yellow, false, 0.0f );
+			}
+		}
+	}
+
+	return bHit;
+}
+
+
+bool FMeshEditorMode::CheckEdge( const bool bUseSphere, const FSphere& Sphere, const float SphereFuzzyDistance, bool bUseRay, const FVector& RayStart, const FVector& RayEnd, const float RayFuzzyDistance, const FVector EdgeVertexPositions[ 2 ], const FVector& CameraLocation, const float FuzzyDistanceScaleFactor, EInteractorShape& ClosestInteractorShape, float& ClosestDistanceToRay, float& ClosestDistanceOnRay, FVector& ClosestHitLocation, const bool bAlreadyHitEdge )
+{
+	bool bHit = false;
+
+	if( bUseSphere )
+	{
+		const float DistanceToSphere = FMath::PointDistToSegment( Sphere.Center, EdgeVertexPositions[ 0 ], EdgeVertexPositions[ 1 ] );
+		if( DistanceToSphere <= Sphere.W )
+		{
+			const FVector ClosestPointOnEdge = FMath::ClosestPointOnSegment( Sphere.Center, EdgeVertexPositions[ 0 ], EdgeVertexPositions[ 1 ] );
+			const float DistanceToCamera = ( CameraLocation - ClosestPointOnEdge ).Size();
+			const float DistanceBasedScaling = MeshEd::OverlayDistanceBias->GetFloat() + DistanceToCamera * FuzzyDistanceScaleFactor;
+
+			// Inside sphere
+			if( DistanceToSphere < ClosestDistanceToRay ||
+				( !bAlreadyHitEdge && FMath::Abs( DistanceToSphere - ClosestDistanceToRay ) < SphereFuzzyDistance * DistanceBasedScaling ) )
+			{
+				ClosestDistanceToRay = DistanceToSphere;
+				ClosestDistanceOnRay = 0.0f;
+				ClosestHitLocation = ClosestPointOnEdge;
+				ClosestInteractorShape = EInteractorShape::GrabberSphere;
+
+				bHit = true;
+			}
+		}
+	}
+
+
+	if( bUseRay )
+	{
+		FVector ClosestPointOnEdge, ClosestPointOnRay;
+		FMath::SegmentDistToSegmentSafe(
+			EdgeVertexPositions[ 0 ], EdgeVertexPositions[ 1 ],
+			RayStart, RayEnd,
+			/* Out */ ClosestPointOnEdge,
+			/* Out */ ClosestPointOnRay );
+		const float DistanceToRay = ( ClosestPointOnEdge - ClosestPointOnRay ).Size();
+		const float DistanceOnRay = ( ClosestPointOnRay - RayStart ).Size();
+
+		const FVector RayDirection = ( RayEnd - RayStart ).GetSafeNormal();
+		const FVector DirectionTowardClosestPointOnRay = ( ClosestPointOnRay - RayStart ).GetSafeNormal();
+		const bool bIsBehindRay = FVector::DotProduct( RayDirection, DirectionTowardClosestPointOnRay ) < 0.0f;
+		if( !bIsBehindRay )
+		{
+			const float DistanceToCamera = ( CameraLocation - ClosestPointOnEdge ).Size();
+			const float DistanceBasedScaling = MeshEd::OverlayDistanceBias->GetFloat() + DistanceToCamera * FuzzyDistanceScaleFactor;
+			check( DistanceBasedScaling > 0.0f );
+
+			// Are we within the minimum distance for hitting an edge?
+			if( DistanceToRay < RayFuzzyDistance * DistanceBasedScaling )
+			{
+				const bool bWithinFuzzyRadius = FMath::Abs( DistanceOnRay - ClosestDistanceOnRay ) < RayFuzzyDistance * DistanceBasedScaling;
+
+				if( ( bWithinFuzzyRadius && DistanceToRay < ClosestDistanceToRay ) || ( !bWithinFuzzyRadius && DistanceOnRay < ClosestDistanceOnRay ) )
+				{
+					ClosestDistanceToRay = DistanceToRay;
+					ClosestDistanceOnRay = DistanceOnRay;
+					ClosestHitLocation = ClosestPointOnRay;
+					ClosestInteractorShape = EInteractorShape::Laser;
+					bHit = true;
+				}
+
+				// @todo mesheditor debug
+				// const float Radius = GHackComponentToWorld.InverseTransformVector( FVector( RayFuzzyDistance * DistanceBasedScaling, 0.0f, 0.0f ) ).X;
+				// DrawDebugSphere( GHackVWI->GetWorld(), GHackComponentToWorld.InverseTransformPosition( ClosestPointOnEdge ), Radius, 12, bHit ? FColor::Green : FColor::Yellow, false, 0.0f );
+			}
+		}
+	}
+
+	return bHit;
+}
+
+
+bool FMeshEditorMode::CheckTriangle( const bool bUseSphere, const FSphere& Sphere, const float SphereFuzzyDistance, bool bUseRay, const FVector& RayStart, const FVector& RayEnd, const float RayFuzzyDistance, const FVector TriangleVertexPositions[ 3 ], const FVector& CameraLocation, const float FuzzyDistanceScaleFactor, EInteractorShape& ClosestInteractorShape, float& ClosestDistanceToRay, float& ClosestDistanceOnRay, FVector& ClosestHitLocation, const bool bAlreadyHitTriangle )
+{
+	bool bHit = false;
+
+	if( bUseSphere )
+	{
+		// @todo grabber: FMath::ClosestPointOnTriangleToPoint doesn't work with degenerates (always returns ray start point?)
+		const FVector ClosestPointOnTriangleToSphere = FMath::ClosestPointOnTriangleToPoint( Sphere.Center, TriangleVertexPositions[ 0 ], TriangleVertexPositions[ 1 ], TriangleVertexPositions[ 2 ] );
+		const float DistanceToSphere = ( ClosestPointOnTriangleToSphere - Sphere.Center ).Size();
+		if( DistanceToSphere <= Sphere.W )
+		{
+			const float DistanceToCamera = ( CameraLocation - ClosestPointOnTriangleToSphere ).Size();
+			const float DistanceBasedScaling = MeshEd::OverlayDistanceBias->GetFloat() + DistanceToCamera * FuzzyDistanceScaleFactor;
+
+			// Inside sphere
+			if( DistanceToSphere < ClosestDistanceToRay ||
+				( !bAlreadyHitTriangle && FMath::Abs( DistanceToSphere - ClosestDistanceToRay ) < SphereFuzzyDistance * DistanceBasedScaling ) )
+			{
+				ClosestHitLocation = ClosestPointOnTriangleToSphere;
+				ClosestDistanceToRay = DistanceToSphere;
+				ClosestDistanceOnRay = 0.0f;
+				ClosestInteractorShape = EInteractorShape::GrabberSphere;
+
+				bHit = true;
+			}
+		}
+	}
+
+
+	if( bUseRay )
+	{
+		// @todo mesheditor: We have like 5 different versions of this in the engine, but nothing generic in a nice place
+		struct Local
+		{
+			static bool RayIntersectTriangle( const FVector& Start, const FVector& End, const FVector& A, const FVector& B, const FVector& C, FVector& IntersectPoint )
+			{
+				const FVector TriNormal = ( B - A ) ^ ( C - A );
+
+				bool bCollide = FMath::SegmentPlaneIntersection( Start, End, FPlane( A, TriNormal ), IntersectPoint );
+				if( !bCollide )
+				{
+					return false;
+				}
+
+				// Make sure points are not colinear.  ComputeBaryCentric2D() doesn't like that.
+				if( TriNormal.SizeSquared() > SMALL_NUMBER )
+				{
+					FVector BaryCentric = FMath::ComputeBaryCentric2D( IntersectPoint, A, B, C );
+					if( BaryCentric.X > 0.0f && BaryCentric.Y > 0.0f && BaryCentric.Z > 0.0f )
+					{
+						return true;
+					}
+				}
+				return false;
+			}
+		};
+
+
+		// @todo mesheditor: Possibly we shouldn't always check for faces when in wire frame view mode?
+
+		// Note: Polygon is assumed to be front facing
+		FVector IntersectionPoint;
+
+		// @todo mesheditor hole: Needs support for polygon hole contours
+		// @todo mesheditor perf: We also have a SIMD version of this that does four triangles at once
+		if( Local::RayIntersectTriangle( RayStart, RayEnd, TriangleVertexPositions[ 0 ], TriangleVertexPositions[ 1 ], TriangleVertexPositions[ 2 ], /* Out */ IntersectionPoint ) )
+		{
+			const float DistanceToCamera = ( CameraLocation - IntersectionPoint ).Size();
+			const float DistanceBasedScaling = MeshEd::OverlayDistanceBias->GetFloat() + DistanceToCamera * FuzzyDistanceScaleFactor;
+    
+			const float DistanceToRay = 0.0f;  // We intersected the triangle, otherwise we wouldn't even be in here
+			const float DistanceOnRay = ( IntersectionPoint - RayStart ).Size();
+			if( DistanceOnRay < ClosestDistanceOnRay ||
+				( !bAlreadyHitTriangle && FMath::Abs( DistanceOnRay - ClosestDistanceOnRay ) < RayFuzzyDistance * DistanceBasedScaling ) )
+			{
+				ClosestHitLocation = IntersectionPoint;
+				ClosestDistanceToRay = DistanceToRay;
+				ClosestDistanceOnRay = DistanceOnRay;
+				ClosestInteractorShape = EInteractorShape::Laser;
+
+				bHit = true;
+			}
+
+			// @todo mesheditor debug
+			// const float Radius = GHackComponentToWorld.InverseTransformVector( FVector( RayFuzzyDistance * DistanceBasedScaling, 0.0f, 0.0f ) ).X;
+			// DrawDebugSphere( GHackVWI->GetWorld(), GHackComponentToWorld.TransformPosition( IntersectionPoint ), Radius, 16, bHit ? FColor::Green : FColor::Yellow, false, 0.0f );
+		}
+	}
+
+	return bHit;
+}
+
+
+
+void FMeshEditorMode::SetMeshElementSelectionMode( EEditableMeshElementType ElementType )
+{
+	const FScopedTransaction Transaction( LOCTEXT( "ChangeMeshElementSelectionMode", "Change Mesh Element Selection Mode" ) );
+	FSetElementSelectionModeChangeInput ChangeInput;
+	ChangeInput.Mode = ElementType;
+	StoreUndo( MeshEditorModeProxyObject, FSetElementSelectionModeChange( MoveTemp( ChangeInput ) ).Execute( MeshEditorModeProxyObject ) );
+}
+
+
+void FMeshEditorMode::RegisterWithExternalMenuSystem()
+{
+	if (VREditorMode != nullptr && VREditorMode->IsFullyInitialized())
+	{
+		FOnRadialMenuGenerated MeshEditActions;
+		MeshEditActions.BindRaw(this, &FMeshEditorMode::MeshEditActionsGenerator);
+		VREditorMode->SetActionsMenuGenerator(MeshEditActions, LOCTEXT("MeshActions", "Mesh Actions"));
+	}
+}
+
+void FMeshEditorMode::UnregisterWithExternalMenuSystem()
+{
+	if (VREditorMode != nullptr && VREditorMode->IsFullyInitialized())
+	{
+		VREditorMode->ResetActionsMenuGenerator();
+	}
+}
+
+int32 FMeshEditorMode::GetSelectedMeshElementIndex(const FMeshElement& MeshElement) const
+{
+	int32 FoundSelectedElementIndex = INDEX_NONE;
+
+	if( MeshElement.ElementAddress.ElementType == GetSelectedMeshElementType() )
+	{
+		for( int32 SelectedElementIndex = 0; SelectedElementIndex < SelectedMeshElements.Num(); ++SelectedElementIndex )
+		{
+			const FMeshElement& SelectedMeshElement = SelectedMeshElements[ SelectedElementIndex ];
+			if( SelectedMeshElement.IsSameMeshElement( MeshElement ) )
+			{
+				FoundSelectedElementIndex = SelectedElementIndex;
+				break;
+			}
+		}
+	}
+
+	return FoundSelectedElementIndex;
+}
+
+
+EEditableMeshElementType FMeshEditorMode::GetSelectedMeshElementType() const
+{
+	// All elements in the list MUST be of the same type, so we simply return the type of the first element
+	return SelectedMeshElements.Num() > 0 ? SelectedMeshElements[ 0 ].ElementAddress.ElementType : EEditableMeshElementType::Invalid;
+}
+
+
+TUniquePtr<FChange> FMeshEditorMode::ClearInvalidSelectedElements_Internal()
+{
+	FSelectOrDeselectMeshElementsChangeInput RevertInput;
+
+	for( const FMeshElement& MeshElement : SelectedMeshElements )
+	{
+		bool bIsValid = false;
+		if( MeshElement.IsValidMeshElement() )
+		{
+			const UEditableMesh* EditableMesh = FindOrCreateEditableMesh( *MeshElement.Component, MeshElement.ElementAddress.SubMeshAddress );
+			if( EditableMesh != nullptr )
+			{
+				if( IsElementIDValid( MeshElement, EditableMesh ) )
+				{
+					bIsValid = true;
+				}
+			}
+		}
+
+		if( !bIsValid )
+		{
+			// @todo mesheditor debug
+			//GWarn->Logf( TEXT( "INVALID ELEMENT deselected: %s" ), *MeshElement.ToString() );
+			RevertInput.MeshElementsToDeselect.Add( MeshElement );
+		}
+	}
+
+	return RevertInput.MeshElementsToDeselect.Num() > 0 ? ( FSelectOrDeselectMeshElementsChange( MoveTemp( RevertInput ) ).Execute( MeshEditorModeProxyObject ) ) : nullptr;
+}
+
+
+void FMeshEditorMode::ClearInvalidHoveredElements()
+{
+	for( FMeshEditorInteractorData& MeshEditorInteractorData : MeshEditorInteractorDatas )
+	{
+		bool bIsValid = false;
+		if( MeshEditorInteractorData.HoveredMeshElement.IsValidMeshElement() )
+		{
+			const UEditableMesh* EditableMesh = FindOrCreateEditableMesh( *MeshEditorInteractorData.HoveredMeshElement.Component, MeshEditorInteractorData.HoveredMeshElement.ElementAddress.SubMeshAddress );
+			if( EditableMesh != nullptr )
+			{
+				if( IsElementIDValid( MeshEditorInteractorData.HoveredMeshElement, EditableMesh ) )
+				{
+					bIsValid = true;
+				}
+			}
+		}
+
+		if( !bIsValid )
+		{
+			MeshEditorInteractorData.HoveredMeshElement = FMeshElement();
+		}
+	}
+
+	for( int32 MeshElementIndex = 0; MeshElementIndex < FadingOutHoveredMeshElements.Num(); ++MeshElementIndex )
+	{
+		const FMeshElement& MeshElement = FadingOutHoveredMeshElements[ MeshElementIndex ];
+
+		bool bIsValid = false;
+		if( MeshElement.IsValidMeshElement() )
+		{
+			const UEditableMesh* EditableMesh = FindOrCreateEditableMesh( *MeshElement.Component, MeshElement.ElementAddress.SubMeshAddress );
+			if( EditableMesh != nullptr )
+			{
+				if( IsElementIDValid( MeshElement, EditableMesh ) )
+				{
+					bIsValid = true;
+				}
+			}
+		}
+
+		if( !bIsValid )
+		{
+			FadingOutHoveredMeshElements.RemoveAt( MeshElementIndex-- );
+		}
+	}
+}
+
+
+void FMeshEditorMode::OnViewportInteractionInputAction( FEditorViewportClient& ViewportClient, UViewportInteractor* ViewportInteractor, const FViewportActionKeyInput& Action, bool& bOutIsInputCaptured, bool& bWasHandled )
+{
+	if( !bWasHandled && Action.ActionType == ViewportWorldActionTypes::SelectAndMove )
+	{
+		UpdateCameraToWorldTransform( ViewportClient );
+
+		const FMeshEditorInteractorData& MeshEditorInteractorData = GetMeshEditorInteractorData( ViewportInteractor );
+
+		// If we're interactively editing something, clicking will commit that change
+		if( Action.Event == IE_Pressed &&
+			ActiveAction != EMeshEditAction::None )
+		{
+			// We're busy doing something else right now.  It might be an interactor trying to click while a different one is in the middle of something.
+			bWasHandled = true;
+		}
+
+		// Otherwise, go ahead and try to interact with what's under the interactor
+		else if( Action.Event == IE_Pressed &&
+				 !bOutIsInputCaptured &&
+				 ActiveAction == EMeshEditAction::None )	// Only if we're not already doing something
+		{
+			bool bWantToStartMoving = false;
+
+			if( EquippedVertexAction == EMeshEditAction::DrawVertices ||
+				EquippedEdgeAction == EMeshEditAction::DrawVertices ||
+				EquippedPolygonAction == EMeshEditAction::DrawVertices )
+			{
+				DrawnPoints.Reset();
+			
+				const bool bActionNeedsHoverLocation = false;
+				StartAction( EMeshEditAction::DrawVertices, ViewportInteractor, bActionNeedsHoverLocation, LOCTEXT( "DrawVertices", "Draw Vertices" ) );
+
+				bOutIsInputCaptured = true;
+				bWasHandled = true;
+			}
+
+			else if( MeshEditorInteractorData.HoveredMeshElement.IsValidMeshElement() && 
+				( MeshEditorInteractorData.bLaserIsValid || MeshEditorInteractorData.bGrabberSphereIsValid ) )
+			{
+				// Make sure the actor is selected
+				// @todo mesheditor: Do we need/want to automatically select actors when doing mesh editing?  If so, consider how undo will 
+				// encapsulate the actor selection change with the mesh element selection change
+				if( false )
+				{
+					UPrimitiveComponent* Component = MeshEditorInteractorData.HoveredMeshElement.Component.Get();
+					if( Component == nullptr || !GEditor->GetSelectedActors()->IsSelected( Component->GetOwner() ) )
+					{
+						GEditor->SelectNone( true, true );
+					}
+					else
+					{
+						GEditor->SelectActor( MeshEditorInteractorData.HoveredMeshElement.Component->GetOwner(), true, true );
+					}
+				}
+
+				// Holding down Control enables multi-select (adds to selection, or deselects single elements when already selected)
+				const bool bIsMultiSelecting = ViewportInteractor->IsModifierPressed();
+
+				const int32 AlreadySelectedMeshElement = GetSelectedMeshElementIndex( MeshEditorInteractorData.HoveredMeshElement );
+				if( AlreadySelectedMeshElement != INDEX_NONE && !bIsMultiSelecting )
+				{
+					const EEditableMeshElementType SelectedMeshElementType = GetSelectedMeshElementType();
+
+					if( SelectedMeshElementType == EEditableMeshElementType::Vertex && EquippedVertexAction == EMeshEditAction::Move )
+					{
+						bWantToStartMoving = true;
+						const bool bActionNeedsHoverLocation = false;
+						StartAction( EMeshEditAction::Move, ViewportInteractor, bActionNeedsHoverLocation, LOCTEXT( "UndoDragVertex", "Drag Vertex" ) );
+					}
+					else if( SelectedMeshElementType == EEditableMeshElementType::Edge && EquippedEdgeAction == EMeshEditAction::Move )
+					{
+						bWantToStartMoving = true;
+						const bool bActionNeedsHoverLocation = false;
+						StartAction( EMeshEditAction::Move, ViewportInteractor, bActionNeedsHoverLocation, LOCTEXT( "UndoDragEdge", "Drag Edge" ) );
+
+					}
+					else if( SelectedMeshElementType == EEditableMeshElementType::Polygon && EquippedPolygonAction == EMeshEditAction::Move )
+					{
+						bWantToStartMoving = true;
+						const bool bActionNeedsHoverLocation = false;
+						StartAction( EMeshEditAction::Move, ViewportInteractor, bActionNeedsHoverLocation, LOCTEXT( "UndoDragPolygon", "Drag Polygon" ) );
+					}
+					// @todo mesheditor: Need a "Shift+Click+Drag" extrude instead of having to explicitly equip the action
+					else if( SelectedMeshElementType == EEditableMeshElementType::Polygon && EquippedPolygonAction == EMeshEditAction::ExtrudePolygon )
+					{
+						bool bHaveExtrudeAxis = false;
+						{
+							const FMeshElement& PolygonElement = MeshEditorInteractorData.HoveredMeshElement;
+
+							if( PolygonElement.IsValidMeshElement() &&
+								PolygonElement.ElementAddress.ElementType == EEditableMeshElementType::Polygon )
+							{
+								// Is it selected?
+								if( GetSelectedMeshElementIndex( PolygonElement ) != INDEX_NONE )
+								{
+									UEditableMesh* EditableMesh = FindOrCreateEditableMesh( *PolygonElement.Component, PolygonElement.ElementAddress.SubMeshAddress );
+									if( EditableMesh != nullptr )
+									{
+										if( IsElementIDValid( PolygonElement, EditableMesh ) )
+										{
+											this->ExtrudePolygonAxisOrigin = MeshEditorInteractorData.HoverLocation;
+
+											// Use the polygon normal as the extrude axis direction
+											{
+												UPrimitiveComponent* Component = PolygonElement.Component.Get();
+												check( Component != nullptr );
+												const FMatrix ComponentToWorldMatrix = Component->GetRenderMatrix();
+
+												const FPolygonRef PolygonRef( PolygonElement.ElementAddress.SectionID, FPolygonID( PolygonElement.ElementAddress.ElementID ) );
+												const FVector ComponentSpacePolygonNormal = EditableMesh->ComputePolygonNormal( PolygonRef );
+
+												const FVector WorldSpacePolygonNormal = ComponentToWorldMatrix.TransformVector( ComponentSpacePolygonNormal ).GetSafeNormal();
+
+												this->ExtrudePolygonAxisDirection = WorldSpacePolygonNormal;
+											}
+
+											bHaveExtrudeAxis = true;
+										}
+									}
+								}
+							}
+						}
+
+						if( bHaveExtrudeAxis )
+						{
+							const bool bActionNeedsHoverLocation = false;
+							StartAction( EMeshEditAction::ExtrudePolygon, ViewportInteractor, bActionNeedsHoverLocation, LOCTEXT( "UndoExtrudePolygon", "Extrude Polygon" ) );
+							bOutIsInputCaptured = true;
+						}
+					}
+					else if( SelectedMeshElementType == EEditableMeshElementType::Polygon && EquippedPolygonAction == EMeshEditAction::FreelyExtrudePolygon )
+					{
+						bWantToStartMoving = true;
+						const bool bActionNeedsHoverLocation = false;
+						StartAction( EMeshEditAction::FreelyExtrudePolygon, ViewportInteractor, bActionNeedsHoverLocation, LOCTEXT( "UndoFreelyExtrudePolygon", "Freely Extrude Polygon" ) );
+					}
+					else if( SelectedMeshElementType == EEditableMeshElementType::Polygon && 
+							 ( EquippedPolygonAction == EMeshEditAction::InsetPolygon || EquippedPolygonAction == EMeshEditAction::BevelPolygon ) )
+					{
+						this->InsetUsingPolygonElement = FMeshElement();
+						{
+							const FMeshElement& PolygonElement = MeshEditorInteractorData.HoveredMeshElement;
+
+							if( PolygonElement.IsValidMeshElement() &&
+								PolygonElement.ElementAddress.ElementType == EEditableMeshElementType::Polygon )
+							{
+								// Is it selected?
+								if( GetSelectedMeshElementIndex( PolygonElement ) != INDEX_NONE )
+								{
+									UEditableMesh* EditableMesh = FindOrCreateEditableMesh( *PolygonElement.Component, PolygonElement.ElementAddress.SubMeshAddress );
+									if( EditableMesh != nullptr )
+									{
+										if( IsElementIDValid( PolygonElement, EditableMesh ) )
+										{
+											this->InsetUsingPolygonElement = PolygonElement;
+										}
+									}
+								}
+							}
+						}
+
+						if( InsetUsingPolygonElement.IsValidMeshElement() )
+						{
+							const bool bActionNeedsHoverLocation = true;
+							if( EquippedPolygonAction == EMeshEditAction::InsetPolygon )
+							{
+								StartAction( EMeshEditAction::InsetPolygon, ViewportInteractor, bActionNeedsHoverLocation, LOCTEXT( "UndoInsetPolygon", "Inset Polygon" ) );
+							}
+							else if( ensure( EquippedPolygonAction == EMeshEditAction::BevelPolygon ) )
+							{
+								StartAction( EMeshEditAction::BevelPolygon, ViewportInteractor, bActionNeedsHoverLocation, LOCTEXT( "UndoBevelPolygon", "Bevel Polygon" ) );
+							}
+
+							bOutIsInputCaptured = true;
+						}
+					}
+					else if( SelectedMeshElementType == EEditableMeshElementType::Edge && EquippedEdgeAction == EMeshEditAction::SplitEdge )
+					{
+						const bool bActionNeedsHoverLocation = true;
+						StartAction( EMeshEditAction::SplitEdge, ViewportInteractor, bActionNeedsHoverLocation, LOCTEXT( "UndoSplitEdge", "Split Edge" ) );
+
+						bOutIsInputCaptured = true;
+					}
+					else if( SelectedMeshElementType == EEditableMeshElementType::Edge && EquippedEdgeAction == EMeshEditAction::SplitEdgeAndDragVertex )
+					{
+						// Figure out what to split
+						GetSelectedMeshesAndEdges( /* Out */ SplitEdgeMeshesAndEdgesToSplit );
+
+						if( SplitEdgeMeshesAndEdgesToSplit.Num() > 0 )
+						{
+							for( auto& MeshAndEdges : SplitEdgeMeshesAndEdgesToSplit )
+							{
+								UEditableMesh* EditableMesh = MeshAndEdges.Key;
+								const TArray<TTuple<FMeshElement, FEdgeID>>& Edges = MeshAndEdges.Value;
+
+								// Figure out where to split
+								FindEdgeSplitUnderInteractor( ViewportInteractor, EditableMesh, Edges, /* Out */ SplitEdgeSplitList );
+
+								// Split the edges
+								if( SplitEdgeSplitList.Num() > 0 )
+								{
+									bWantToStartMoving = true;
+									const bool bActionNeedsHoverLocation = false;
+									StartAction( EMeshEditAction::SplitEdgeAndDragVertex, ViewportInteractor, bActionNeedsHoverLocation, LOCTEXT( "UndoSplitEdgeAndDragVertex", "Split Edge And Drag Vertex" ) );
+
+									// We just need to get one split list.  We'll use that for all selected meshes.
+									break;
+								}
+							}
+						}
+					}
+					else if( SelectedMeshElementType == EEditableMeshElementType::Edge && EquippedEdgeAction == EMeshEditAction::InsertEdgeLoop )
+					{
+						const bool bActionNeedsHoverLocation = true;
+						StartAction( EMeshEditAction::InsertEdgeLoop, ViewportInteractor, bActionNeedsHoverLocation, LOCTEXT( "UndoInsertEdgeLoop", "Insert Edge Loop" ) );
+
+						bOutIsInputCaptured = true;
+					}
+					else if( SelectedMeshElementType == EEditableMeshElementType::Edge && EquippedEdgeAction == EMeshEditAction::ExtendEdge )
+					{
+						bWantToStartMoving = true;
+						const bool bActionNeedsHoverLocation = false;
+						StartAction( EMeshEditAction::ExtendEdge, ViewportInteractor, bActionNeedsHoverLocation, LOCTEXT( "UndoExtendEdge", "Extend Edge" ) );
+					}
+					else if( SelectedMeshElementType == EEditableMeshElementType::Vertex && EquippedVertexAction == EMeshEditAction::ExtendVertex )
+					{
+						bWantToStartMoving = true;
+						const bool bActionNeedsHoverLocation = false;
+						StartAction( EMeshEditAction::ExtendVertex, ViewportInteractor, bActionNeedsHoverLocation, LOCTEXT( "UndoExtendVertex", "Extend Vertex" ) );
+					}
+					else if( SelectedMeshElementType == EEditableMeshElementType::Vertex && EquippedVertexAction == EMeshEditAction::EditVertexCornerSharpness )
+					{
+						EditSharpnessStartLocation = MeshEditorInteractorData.HoverLocation;
+
+						const bool bActionNeedsHoverLocation = false;
+						StartAction( EMeshEditAction::EditVertexCornerSharpness, ViewportInteractor, bActionNeedsHoverLocation, LOCTEXT( "UndoEditVertexCornerSharpness", "Edit Vertex Corner Sharpness" ) );
+
+						bOutIsInputCaptured = true;
+					}
+					else if( SelectedMeshElementType == EEditableMeshElementType::Edge && EquippedEdgeAction == EMeshEditAction::EditEdgeCreaseSharpness )
+					{
+						EditSharpnessStartLocation = MeshEditorInteractorData.HoverLocation;
+
+						const bool bActionNeedsHoverLocation = false;
+						StartAction( EMeshEditAction::EditEdgeCreaseSharpness, ViewportInteractor, bActionNeedsHoverLocation, LOCTEXT( "UndoEditEdgeCreaseSharpness", "Edit Edge Crease Sharpness" ) );
+
+						bOutIsInputCaptured = true;
+					}
+				}
+				else 
+				{
+					if( AlreadySelectedMeshElement != INDEX_NONE && bIsMultiSelecting )
+					{
+						// Deselect it
+						const FScopedTransaction Transaction( LOCTEXT( "DeselectMeshElements", "Deselect Element" ) );
+
+						FSelectOrDeselectMeshElementsChangeInput ChangeInput;
+						ChangeInput.MeshElementsToDeselect.Add( SelectedMeshElements[ AlreadySelectedMeshElement ] );
+						StoreUndo( MeshEditorModeProxyObject, FSelectOrDeselectMeshElementsChange( MoveTemp( ChangeInput ) ).Execute( MeshEditorModeProxyObject ) );
+					}
+					else if( MeshElementSelectionMode == EEditableMeshElementType::Any || MeshElementSelectionMode == MeshEditorInteractorData.HoveredMeshElement.ElementAddress.ElementType )
+					{
+						// Start painting selection
+						const bool bIsSelectByPaintingEnabled = MeshEd::EnableSelectByPainting->GetInt() != 0;
+						if( bIsSelectByPaintingEnabled )
+						{
+							const bool bActionNeedsHoverLocation = true;
+							StartAction( EMeshEditAction::SelectByPainting, ViewportInteractor, bActionNeedsHoverLocation, LOCTEXT( "UndoSelectingMeshElements", "Select Element" ) );
+							bOutIsInputCaptured = true;
+						}
+					
+						FSelectOrDeselectMeshElementsChangeInput ChangeInput;
+
+						// Unless we're trying to multi-select, clear selection before selecting something new
+						if( !bIsMultiSelecting )
+						{
+							ChangeInput.MeshElementsToDeselect = SelectedMeshElements;
+						}
+
+						// Select the element under the mouse cursor
+						ChangeInput.MeshElementsToSelect.Add( MeshEditorInteractorData.HoveredMeshElement );
+
+						TUniquePtr<FChange> RevertChange = FSelectOrDeselectMeshElementsChange( MoveTemp( ChangeInput ) ).Execute( MeshEditorModeProxyObject );
+
+						if( bIsSelectByPaintingEnabled )
+						{
+							SelectingByPaintingRevertChangeInput = MakeUnique<FCompoundChangeInput>();
+							SelectingByPaintingRevertChangeInput->Subchanges.Add( MoveTemp( RevertChange ) );
+						}
+						else
+						{
+							// If select by painting is disabled, add a transaction immediately
+							const FScopedTransaction Transaction( LOCTEXT( "SelectElement", "Select Element" ) );
+							StoreUndo( MeshEditorModeProxyObject, MoveTemp( RevertChange ) );
+						}
+					}
+				}
+
+				bWasHandled = true;
+			}
+
+			if( bWantToStartMoving )
+			{
+				UPrimitiveComponent* ClickedTransformGizmoComponent = nullptr;
+				const bool bIsPlacingNewObjects = false;
+				const bool bAllowInterpolationWhenPlacing = true;
+				const bool bStartTransaction = false;
+				const bool bWithGrabberSphere = ( MeshEditorInteractorData.HoverInteractorShape == EInteractorShape::GrabberSphere );
+				ViewportWorldInteraction->StartDragging(
+					ActiveActionInteractor,
+					ClickedTransformGizmoComponent,
+					MeshEditorInteractorData.HoverLocation,
+					bIsPlacingNewObjects,
+					bAllowInterpolationWhenPlacing,
+					bStartTransaction,
+					bWithGrabberSphere );
+
+				// NOTE: We purposely don't set bIsInputCaptured=true here, because ViewportWorldInteraction will take over handling
+				//		 of the 'release' input event for this drag
+				// ...
+			}
+		}
+		else if( Action.Event == IE_Released )
+		{
+			if( ActiveAction != EMeshEditAction::None && 
+				ActiveAction != EMeshEditAction::MoveUsingGizmo &&	// The button 'release' for gizmo-based movement is handled by the viewport world interaction system
+				bOutIsInputCaptured )
+			{
+				if( ActiveActionInteractor == nullptr || ActiveActionInteractor == ViewportInteractor )
+				{
+					if( ActiveAction == EMeshEditAction::SelectByPainting )
+					{
+						check( SelectingByPaintingRevertChangeInput.IsValid() );
+
+						// Did we end up selecting anything?
+						if( SelectingByPaintingRevertChangeInput->Subchanges.Num() > 0 )
+						{
+							// Make sure we still have an active transaction.  It's possible that something strange happened and
+							// we received a release event out of band with where we started it, or some other editor event
+							// canceled our transaction while the mouse was down.
+							if( GUndo != nullptr )
+							{
+								StoreUndo( MeshEditorModeProxyObject, MakeUnique<FCompoundChange>( MoveTemp( *SelectingByPaintingRevertChangeInput.Release() ) ) );
+							}
+						}
+						SelectingByPaintingRevertChangeInput.Reset();
+					}
+
+					FinishAction();
+				}
+
+				bOutIsInputCaptured = false;
+				bWasHandled = true;
+			}
+		}
+	}
+}
+
+
+void FMeshEditorMode::StartAction( EMeshEditAction NewAction, UViewportInteractor* ActionInteractor, const bool bActionNeedsHoverLocation, const FText& UndoText )
+{
+	// Don't start a new action without finishing the previous one!
+	check( ActiveAction == EMeshEditAction::None );
+
+	PlayStartActionSound( NewAction, ActionInteractor );
+
+	ActiveAction = NewAction;
+	ActiveActionInteractor = ActionInteractor;
+	bActiveActionNeedsHoverLocation = bActionNeedsHoverLocation;
+	bIsFirstActiveActionUpdate = true;
+
+	// Start tracking undo state (unless the undo string was empty.)
+	if( !UndoText.IsEmpty() )
+	{
+		TrackingTransaction.TransCount++;
+		TrackingTransaction.Begin( UndoText );
+
+		// Suspend actor/component modification during each delta step to avoid recording unnecessary overhead into the transaction buffer
+		GEditor->DisableDeltaModification( true );
+	}
+}
+
+
+void FMeshEditorMode::FinishAction()
+{
+	// @todo mesheditor: Make sure this is called before Undo is invoked (PreEditUndo!), otherwise the previous action will be undone instead of the active one
+
+	check( ActiveAction != EMeshEditAction::None );
+	check( !IsUndoSystemAvailable() || GEditor->IsTransactionActive() );	// Someone must have started a transaction! (It might not have been us though.)
+
+	const bool bIsActionFinishing = true;
+
+	if( ActiveAction != EMeshEditAction::SelectByPainting )
+	{
+		UpdateActiveAction( bIsActionFinishing );
+	}
+
+	if( ActiveAction == EMeshEditAction::DrawVertices )
+	{
+		// @todo mesheditor: Drawing vertices will likely need to be a different kind of active action as it works differently to the others.
+		// For now, this just forces vertex drawing to be a "one shot" kind of mode.
+		EquippedVertexAction = EMeshEditAction::Move;
+		EquippedEdgeAction = EMeshEditAction::Move;
+		EquippedPolygonAction = EMeshEditAction::Move;
+	}
+
+	if( bIsActionFinishing )
+	{
+		PlayFinishActionSound( ActiveAction, ActiveActionInteractor );
+	}
+
+	ActiveAction = EMeshEditAction::None;
+	ActiveActionInteractor = nullptr;
+	bActiveActionNeedsHoverLocation = false;
+
+	if( TrackingTransaction.IsActive() )
+	{
+		--TrackingTransaction.TransCount;
+		TrackingTransaction.End();
+		GEditor->DisableDeltaModification( false );
+	}
+
+	// If the action has finished, make sure the gizmo is in the correct place as elements may have moved.
+	if( bIsActionFinishing )
+	{
+		RefreshTransformables();
+	}
+}
+
+
+bool FMeshEditorMode::IsUndoSystemAvailable()
+{
+	return GUndo != nullptr;
+}
+
+
+void FMeshEditorMode::StoreUndo( UObject* Object, TUniquePtr<FChange> UndoChange )
+{
+	// Did you forget to use an FScopedTransaction?  If GUndo was null, then most likely we forgot to wrap this call within an editor transaction.
+	// The only exception is in Simulate mode, where Undo is not allowed.
+	check( IsUndoSystemAvailable() || GEditor == nullptr || GEditor->bIsSimulatingInEditor );
+
+	if( IsUndoSystemAvailable() )
+	{
+		GUndo->StoreUndo( Object, MoveTemp( UndoChange ) );
+	}	
+}
+
+
+void FMeshEditorMode::PostUndo()
+{
+	// Make sure we're not hovering over anything that was deleted after an undo action
+	ClearInvalidHoveredElements();
+
+	// Update our transformable list
+	RefreshTransformables();
+}
+
+
+bool FMeshEditorMode::FrustumSelect( const FConvexVolume& InFrustum, FEditorViewportClient* InViewportClient, bool InSelect )
+{
+	// @todo mesheditor urgent: settings class for bundling together all these kind of things
+	const bool bShouldDeselectAllFirst = true;	// @todo mesheditor: needs to be passed to this method
+	const bool bOnlyInteractWithSelectedObjects = GetDefault<UMeshEditorSettings>()->bOnlyEditSelectedObjects;
+	const bool bOnlySelectVisibleMeshes = GetDefault<UMeshEditorSettings>()->bOnlySelectVisibleMeshes;
+	const bool bOnlySelectVisibleElements = GetDefault<UMeshEditorSettings>()->bOnlySelectVisibleElements;
+
+	UWorld* World = GetWorld();
+
+	UpdateCameraToWorldTransform( *InViewportClient );
+
+	// First obtain a list of candidate editable meshes which intersect with the frustum
+
+	static TArray<TTuple<UPrimitiveComponent*, UEditableMesh*>> CandidateMeshes;
+	CandidateMeshes.Empty();
+
+	// Lambda which creates editable meshes from any eligible component in the actor
+	auto AddEditableMeshFromActor = [ this, &InFrustum ]( AActor* Actor )
+	{
+		if( Actor->IsEditorOnly() && Actor->IsSelectable() )
+		{
+			return;
+		}
+
+		TInlineComponentArray<UPrimitiveComponent*> Components;
+		Actor->GetComponents( Components );
+
+		for( UPrimitiveComponent* Component : Components )
+		{
+			check( Component );
+			if( Component->IsRegistered() &&
+			    Component->IsVisibleInEditor() &&
+				!Component->IsEditorOnly() &&
+			    InFrustum.IntersectBox( Component->Bounds.Origin, Component->Bounds.BoxExtent ) )
+			{
+				const int32 LODIndex = 0;
+				UEditableMesh* EditableMesh = FindOrCreateEditableMesh( *Component, UEditableMeshFactory::MakeSubmeshAddress( Component, LODIndex ) );
+				if( EditableMesh )
+				{
+					CandidateMeshes.Emplace( Component, EditableMesh );
+				}
+			}
+		}
+	};
+
+	// Now find all actors which lie within the selection box and find or create editable meshes for them.
+	// There are two possible paths.
+
+	if( bOnlySelectVisibleMeshes )
+	{
+		// By this method, interrogate the hit proxy to determine which actors are within the selection box
+
+		float StartX = TNumericLimits<float>::Max();
+		float StartY = TNumericLimits<float>::Max();
+		float EndX = TNumericLimits<float>::Lowest();
+		float EndY = TNumericLimits<float>::Lowest();
+
+		// Frustum sides are in the first four indices
+		// Find intersection points and project to screen space to determine the bounding rectangle of the selection box
+		for( int32 PlaneIndex = 0; PlaneIndex < 4; ++PlaneIndex )
+		{
+			const FPlane& Plane1 = InFrustum.Planes[ PlaneIndex ];
+			const FPlane& Plane2 = InFrustum.Planes[ ( PlaneIndex + 1 ) % 4 ];
+			FVector I, D;
+			if( FMath::IntersectPlanes2( I, D, Plane1, Plane2 ) )
+			{
+				FSceneViewFamilyContext ViewFamily(FSceneViewFamily::ConstructionValues(InViewportClient->Viewport, InViewportClient->GetScene(), InViewportClient->EngineShowFlags ));
+				FSceneView* SceneView = InViewportClient->CalcSceneView(&ViewFamily);
+
+				FVector2D V;
+				if( SceneView->WorldToPixel( I, V ) )
+				{
+					StartX = FMath::Min( StartX, V.X );
+					StartY = FMath::Min( StartY, V.Y );
+					EndX = FMath::Max( EndX, V.X );
+					EndY = FMath::Max( EndY, V.Y );
+				}
+			}
+		}
+
+		const int32 ViewportSizeX = InViewportClient->Viewport->GetSizeXY().X;
+		const int32 ViewportSizeY = InViewportClient->Viewport->GetSizeXY().Y;
+		FIntRect BoxRect( FIntPoint( FMath::Max( 0.0f, StartX ), FMath::Max( 0.0f, StartY ) ), FIntPoint( FMath::Min( ViewportSizeX, FMath::TruncToInt( EndX + 1 ) ), FMath::Min( ViewportSizeY, FMath::TruncToInt( EndY + 1) ) ) );
+
+		TSet<AActor*> HitActors;
+		TSet<UModel*> HitModels;
+		InViewportClient->Viewport->GetActorsAndModelsInHitProxy( BoxRect, HitActors, HitModels );
+
+		for( AActor* Actor : HitActors )
+		{
+			if( !bOnlyInteractWithSelectedObjects || GEditor->GetSelectedActors()->IsSelected( Actor ) )
+			{
+				AddEditableMeshFromActor( Actor );
+			}
+		}
+	}
+	else
+	{
+		// Determine actors within the selection box by testing intersections between all candidate actors' bounding boxes and the frustum
+
+		for( FActorIterator It( World ); It; ++It )
+		{
+			AActor* Actor = *It;
+
+			if( !Actor->IsA( ABrush::StaticClass() ) &&
+			    !Actor->IsHiddenEd() &&
+			    ( !bOnlyInteractWithSelectedObjects || GEditor->GetSelectedActors()->IsSelected( Actor ) ) )
+			{
+				AddEditableMeshFromActor( Actor );
+			}
+		}
+	}
+
+	// Now find candidate editable mesh elements.
+
+	MarqueeSelectVertices.Empty();
+	MarqueeSelectEdges.Empty();
+	MarqueeSelectPolygons.Empty();
+
+	for( TTuple<UPrimitiveComponent*, UEditableMesh*> CandidateMesh : CandidateMeshes )
+	{
+		UPrimitiveComponent* Component = CandidateMesh.Get<0>();
+		UEditableMesh* EditableMesh = CandidateMesh.Get<1>();
+		FTransform ComponentTransform = Component->GetComponentTransform();
+
+		static TArray<FEdgeID> SelectedEdgeIDs;
+		SelectedEdgeIDs.Empty();
+
+		static TSet<FPolygonRef> SelectedPolygonRefs;
+		SelectedPolygonRefs.Empty();
+
+		static TSet<FVertexID> SelectedVertexIDs;
+		SelectedVertexIDs.Empty();
+
+		// First, find all edges which lie at least partially within the frustum.
+		const int32 EdgeArraySize = EditableMesh->GetEdgeArraySize();
+		for( int32 EdgeIndex = 0; EdgeIndex < EdgeArraySize; ++EdgeIndex )
+		{
+			const FEdgeID EdgeID( EdgeIndex );
+			if( EditableMesh->IsValidEdge( EdgeID ) )
+			{
+				const FVertexID VertexID1( EditableMesh->GetEdgeVertex( EdgeID, 0 ) );
+				const FVertexID VertexID2( EditableMesh->GetEdgeVertex( EdgeID, 1 ) );
+
+				const FVector VertexPosition1( ComponentTransform.TransformPosition( EditableMesh->GetVertexAttribute( VertexID1, UEditableMeshAttribute::VertexPosition(), 0 ) ) );
+				const FVector VertexPosition2( ComponentTransform.TransformPosition( EditableMesh->GetVertexAttribute( VertexID2, UEditableMeshAttribute::VertexPosition(), 0 ) ) );
+
+				if( InFrustum.IntersectLineSegment( VertexPosition1, VertexPosition2 ) )
+				{
+					bool bAreAllPolysBackFacing = true;
+
+					// Now iterate through all connected polygons.
+					// If any are front facing, we consider the edge also to be front facing.
+					const int32 EdgeConnectedPolygonCount = EditableMesh->GetEdgeConnectedPolygonCount( EdgeID );
+					for( int32 EdgeConnectedPolygonIndex = 0; EdgeConnectedPolygonIndex < EdgeConnectedPolygonCount; ++EdgeConnectedPolygonIndex )
+					{
+						const FPolygonRef EdgeConnectedPolygonID( EditableMesh->GetEdgeConnectedPolygon( EdgeID, EdgeConnectedPolygonIndex ) );
+
+						// Determine whether polygon is back facing or not using dot product of its normal with the direction vector from the eye position to somewhere on the plane
+						// (one of the vertex positions is sufficient for this)
+						const FVector PolyNormal = ComponentTransform.TransformVector( EditableMesh->ComputePolygonNormal( EdgeConnectedPolygonID ) );
+						const FVector ViewDirection = VertexPosition1 - CachedCameraToWorld.GetValue().GetLocation();
+						const bool bIsBackFacing = ( FVector::DotProduct( PolyNormal, ViewDirection ) > 0.0f );
+
+						bAreAllPolysBackFacing &= bIsBackFacing;
+
+						if( !bOnlySelectVisibleElements || !bIsBackFacing )
+						{
+							// Add the polygon if it is front facing, or if we don't care about only selecting visible elements
+							SelectedPolygonRefs.Add( EdgeConnectedPolygonID );
+						}
+					}
+
+					if( !bOnlySelectVisibleElements || !bAreAllPolysBackFacing )
+					{
+						// If at least one of the connected polygons is front facing, we deem the edge also to be front facing
+						SelectedEdgeIDs.Add( EdgeID );
+
+						// Just because the edge is in the frustum doesn't imply that both its constituent vertices are.
+						// We have to do further frustum / point checks.
+						if( InFrustum.IntersectPoint( VertexPosition1 ) )
+						{
+							SelectedVertexIDs.Add( VertexID1 );
+						}
+
+						if( InFrustum.IntersectPoint( VertexPosition2 ) )
+						{
+							SelectedVertexIDs.Add( VertexID2 );
+						}
+					}
+				}
+			}
+		}
+
+		// Next, look for any orphaned vertices (i.e. which do not form part of an edge)
+		const int32 VertexArraySize = EditableMesh->GetVertexArraySize();
+		for( int32 VertexIndex = 0; VertexIndex < VertexArraySize; ++VertexIndex )
+		{
+			const FVertexID VertexID( VertexIndex );
+			if( EditableMesh->IsValidVertex( VertexID ) )
+			{
+				// If the vertex has connected edges, it will have already been considered in the above code.
+				// Here we only want to catch vertices with no associated edges.
+				if( EditableMesh->GetVertexConnectedEdgeCount( VertexID ) == 0 )
+				{
+					const FVector VertexPosition( ComponentTransform.TransformPosition( EditableMesh->GetVertexAttribute( VertexID, UEditableMeshAttribute::VertexPosition(), 0 ) ) );
+					if( InFrustum.IntersectPoint( VertexPosition ) )
+					{
+						// As the vertex is orphaned, it cannot be front or back facing. So we add it regardless.
+						SelectedVertexIDs.Add( VertexID );
+					}
+				}
+			}
+		}
+
+		// Fill arrays with the selected elements
+
+		if( MeshElementSelectionMode == EEditableMeshElementType::Vertex || MeshElementSelectionMode == EEditableMeshElementType::Any )
+		{
+			MarqueeSelectVertices.Reserve( MarqueeSelectVertices.Num() + SelectedVertexIDs.Num() );
+			for( FVertexID SelectedVertexID : SelectedVertexIDs )
+			{
+				MarqueeSelectVertices.Emplace( Component, EditableMesh->GetSubMeshAddress(), SelectedVertexID );
+			}
+		}
+
+		if( MeshElementSelectionMode == EEditableMeshElementType::Edge || MeshElementSelectionMode == EEditableMeshElementType::Any )
+		{
+			MarqueeSelectEdges.Reserve( MarqueeSelectEdges.Num() + SelectedEdgeIDs.Num() );
+			for( FEdgeID SelectedEdgeID : SelectedEdgeIDs )
+			{
+				MarqueeSelectEdges.Emplace( Component, EditableMesh->GetSubMeshAddress(), SelectedEdgeID );
+			}
+		}
+
+		if( MeshElementSelectionMode == EEditableMeshElementType::Polygon || MeshElementSelectionMode == EEditableMeshElementType::Any )
+		{
+			MarqueeSelectPolygons.Reserve( MarqueeSelectPolygons.Num() + SelectedPolygonRefs.Num() );
+			for( FPolygonRef SelectedPolygonRef : SelectedPolygonRefs )
+			{
+				MarqueeSelectPolygons.Emplace( Component, EditableMesh->GetSubMeshAddress(), SelectedPolygonRef );
+			}
+		}
+	}
+
+	if( MeshElementSelectionMode != EEditableMeshElementType::Any )
+	{
+		PerformMarqueeSelect( MeshElementSelectionMode );
+		return true;
+	}
+
+	// If we are in 'any' selection mode, build a context menu to pop up in order to choose which element type the user wishes to select
+
+	const bool bShouldCloseWindowAfterMenuSelection = true;
+	FMenuBuilder MenuBuilder( bShouldCloseWindowAfterMenuSelection, CommonCommands );
+	{
+		MenuBuilder.AddMenuEntry( FMeshEditorCommonCommands::Get().MarqueeSelectVertices );
+		MenuBuilder.AddMenuEntry( FMeshEditorCommonCommands::Get().MarqueeSelectEdges );
+		MenuBuilder.AddMenuEntry( FMeshEditorCommonCommands::Get().MarqueeSelectPolygons );
+	}
+
+	TSharedRef<SWidget> MenuWidget = MenuBuilder.MakeWidget();
+
+	TSharedPtr<SEditorViewport> ViewportWidget = InViewportClient->GetEditorViewportWidget();
+	if( ViewportWidget.IsValid() )
+	{
+		TSharedPtr<IMenu> Menu = FSlateApplication::Get().PushMenu(
+			ViewportWidget.ToSharedRef(),
+			FWidgetPath(),
+			MenuWidget,
+			FSlateApplication::Get().GetCursorPos(),
+			FPopupTransitionEffect( FPopupTransitionEffect::ContextMenu )
+		);
+
+		// Continue the scope of the current transaction while the menu is active.
+		// It will be ended when the menu is dismissed.
+		MarqueeSelectTransaction = MakeUnique<FScopedTransaction>( LOCTEXT( "MarqueeSelectElements", "Marquee Select Elements" ) );
+		bMarqueeSelectTransactionActive = true;
+
+		auto OnMenuDismissed = [ this ]( TSharedRef<IMenu> Menu )
+		{
+			// End transaction here.
+			// This will actually be released in the next Tick() - this is necessary because the OnMenuDismissed callback happens *before* the action has been executed,
+			// and we need the transaction to remain active until afterwards.
+			bMarqueeSelectTransactionActive = false;
+		};
+
+		Menu->GetOnMenuDismissed().AddLambda( OnMenuDismissed );
+	}
+
+	return true;
+}
+
+
+bool FMeshEditorMode::ShouldDrawWidget() const
+{
+	// We draw our own transform gizmo
+	return false;
+}
+
+
+void FMeshEditorMode::PerformMarqueeSelect( EEditableMeshElementType ElementType )
+{
+	const FScopedTransaction Transaction( LOCTEXT( "MarqueeSelectElements", "Marquee Select Elements" ) );
+
+	FSelectOrDeselectMeshElementsChangeInput ChangeInput;
+
+	const bool bShouldDeselectAllFirst = true;
+	if( bShouldDeselectAllFirst )
+	{
+		ChangeInput.MeshElementsToDeselect = SelectedMeshElements;
+	}
+
+	switch( ElementType )
+	{
+		case EEditableMeshElementType::Vertex:
+			ChangeInput.MeshElementsToSelect = MarqueeSelectVertices;
+			break;
+
+		case EEditableMeshElementType::Edge:
+			ChangeInput.MeshElementsToSelect = MarqueeSelectEdges;
+			break;
+
+		case EEditableMeshElementType::Polygon:
+			ChangeInput.MeshElementsToSelect = MarqueeSelectPolygons;
+			break;
+	}
+
+	StoreUndo( MeshEditorModeProxyObject, FSelectOrDeselectMeshElementsChange( MoveTemp( ChangeInput ) ).Execute( MeshEditorModeProxyObject ) );
+}
+
+
+void FMeshEditorMode::RefreshTransformables()
+{
+	// Don't refresh transformables while we're actively moving them around
+	const bool bAllowRefresh =
+		ActiveAction == EMeshEditAction::None ||
+		ActiveAction == EMeshEditAction::SelectByPainting ||
+		bIsFirstActiveActionUpdate;
+	if( !bAllowRefresh )
+	{
+		return;
+	}
+
+	// @todo gizmo: For better performance, we should probably avoid setting up transformables while churning through undo states,
+	//      and instead defer it until the user will actually be able to see the end result
+	//		NOTE:  We also do this in FDeselectAllMeshElementsChange::Execute()
+	TArray<TUniquePtr<FViewportTransformable>> Transformables;
+	for( const FMeshElement& MeshElement : SelectedMeshElements )
+	{
+		if( MeshElement.IsValidMeshElement() )
+		{
+			UEditableMesh* EditableMesh = FindOrCreateEditableMesh( *MeshElement.Component, MeshElement.ElementAddress.SubMeshAddress );
+			if( EditableMesh != nullptr )
+			{
+				if( IsElementIDValid( MeshElement, EditableMesh ) )
+				{
+					UPrimitiveComponent* Component = MeshElement.Component.Get();
+					check( Component != nullptr );
+					const FTransform ComponentToWorld = Component->GetComponentToWorld();
+					const FMatrix ComponentToWorldMatrix = Component->GetRenderMatrix();
+
+					FTransform ElementTransform = FTransform::Identity;
+					switch( MeshElement.ElementAddress.ElementType )
+					{
+						case EEditableMeshElementType::Vertex:
+							ElementTransform.SetLocation(
+								ComponentToWorldMatrix.TransformPosition(
+									EditableMesh->GetVertexAttribute( FVertexID( MeshElement.ElementAddress.ElementID ), UEditableMeshAttribute::VertexPosition(), 0 )
+								)
+							);
+							break;
+
+						case EEditableMeshElementType::Edge:
+						{
+							FVertexID EdgeVertexID0, EdgeVertexID1;
+							EditableMesh->GetEdgeVertices( FEdgeID( MeshElement.ElementAddress.ElementID ), /* Out */ EdgeVertexID0, /* Out */ EdgeVertexID1 );
+
+							const FVector ComponentSpaceEdgeCenter =
+								FMath::Lerp(
+									EditableMesh->GetVertexAttribute( EdgeVertexID0, UEditableMeshAttribute::VertexPosition(), 0 ),
+									EditableMesh->GetVertexAttribute( EdgeVertexID1, UEditableMeshAttribute::VertexPosition(), 0 ),
+									0.5f );
+
+							ElementTransform.SetLocation( ComponentToWorldMatrix.TransformPosition( ComponentSpaceEdgeCenter ) );
+						}
+						break;
+
+						case EEditableMeshElementType::Polygon:
+						{
+							const FPolygonRef PolygonRef( MeshElement.ElementAddress.SectionID, FPolygonID( MeshElement.ElementAddress.ElementID ) );
+
+							const FVector ComponentSpacePolygonCenter = EditableMesh->ComputePolygonCenter( PolygonRef );
+							ElementTransform.SetLocation( ComponentToWorldMatrix.TransformPosition( ComponentSpacePolygonCenter ) );
+
+							const FVector WindingVector =
+								( ComponentToWorldMatrix.TransformPosition( EditableMesh->GetVertexAttribute( EditableMesh->GetPolygonPerimeterVertex( PolygonRef, 1 ), UEditableMeshAttribute::VertexPosition(), 0 ) ) -
+									ComponentToWorldMatrix.TransformPosition( EditableMesh->GetVertexAttribute( EditableMesh->GetPolygonPerimeterVertex( PolygonRef, 0 ), UEditableMeshAttribute::VertexPosition(), 0 ) ) ).GetSafeNormal();
+
+							const FVector PolygonNormal = ComponentToWorld.TransformVectorNoScale( EditableMesh->ComputePolygonNormal( PolygonRef ) ).GetSafeNormal();
+
+							const FVector PolygonBinormal = FVector::CrossProduct( PolygonNormal, WindingVector ).GetSafeNormal();
+							const FVector PolygonTangent = FVector::CrossProduct( PolygonBinormal, PolygonNormal );
+
+							const FQuat PolygonOrientation( FMatrix( PolygonTangent, PolygonBinormal, PolygonNormal, FVector::ZeroVector ).ToQuat() );
+
+							ElementTransform.SetRotation( PolygonOrientation );
+						}
+						break;
+
+						default:
+							check( 0 );
+					}
+
+					FMeshElementViewportTransformable* Transformable = new FMeshElementViewportTransformable( *this );
+
+					Transformables.Add( TUniquePtr<FViewportTransformable>( Transformable ) );
+
+					Transformable->MeshElement = MeshElement;
+					Transformable->CurrentTransform = Transformable->StartTransform = ElementTransform;
+				}
+			}
+		}
+	}
+
+	ViewportWorldInteraction->SetTransformables( MoveTemp( Transformables ) );
+}
+
+
+void FMeshEditorMode::OnActorSelectionChanged( const TArray<UObject*>& NewSelection, bool bForceRefresh )
+{
+	// If we're configured to only interact with selected objects, deselect any elements that no longer belong
+	// to the selected set of actors.
+	const bool bOnlyInteractWithSelectedObjects = GetDefault<UMeshEditorSettings>()->bOnlyEditSelectedObjects;
+	if( bOnlyInteractWithSelectedObjects )
+	{
+		// Don't respond to actor selection changes if a transaction isn't in progress, because it's probably
+		// initiated from an undo/redo action itself, in which case the selection state changes will already
+		// be part of the undo history and we don't need to do anything.
+		if( GEditor->IsTransactionActive() )
+		{
+			FSelectOrDeselectMeshElementsChangeInput ChangeInput;
+
+			for( int32 SelectedElementIndex = 0; SelectedElementIndex < this->SelectedMeshElements.Num(); ++SelectedElementIndex )
+			{
+				const FMeshElement& SelectedMeshElement = this->SelectedMeshElements[ SelectedElementIndex ];
+				if( !SelectedMeshElement.Component.IsValid() || ( SelectedMeshElement.IsValidMeshElement() && !SelectedMeshElement.Component->GetOwner()->IsSelected() ) )
+				{
+					ChangeInput.MeshElementsToDeselect.Add( SelectedMeshElement );
+				}
+			}
+
+			if( ChangeInput.MeshElementsToDeselect.Num() > 0 || ChangeInput.MeshElementsToSelect.Num() > 0 )
+			{
+				StoreUndo( MeshEditorModeProxyObject, FDeselectAllMeshElementsChange( MoveTemp( FDeselectAllMeshElementsChangeInput() ) ).Execute( MeshEditorModeProxyObject ) );
+			}
+		}
+	}
+}
+
+void FMeshEditorMode::MeshEditActionsGenerator(FMenuBuilder MenuBuilder, TSharedPtr<FUICommandList> CommandList, class UVREditorMode* VRMode, float& RadiusOverride)
+{
+	MenuBuilder.AddMenuEntry(
+		LOCTEXT("AddSubdivision", "Add SubD"),
+		FText(),
+		FSlateIcon(FMeshEditorStyle::GetStyleSetName(), "MeshEditorMode.AddSubdivision"),
+		FUIAction
+		(
+			FExecuteAction::CreateSP(this, &FMeshEditorMode::AddOrRemoveSubdivisionLevel, true )
+		),
+		NAME_None,
+		EUserInterfaceActionType::Check
+		);
+	MenuBuilder.AddMenuEntry(
+		LOCTEXT("RemoveSubdivision", "Remove SubD"),
+		FText(),
+		FSlateIcon(FMeshEditorStyle::GetStyleSetName(), "MeshEditorMode.RemoveSubdivision"),
+		FUIAction
+		(
+			FExecuteAction::CreateSP(this, &FMeshEditorMode::AddOrRemoveSubdivisionLevel, false)
+		),
+		NAME_None,
+		EUserInterfaceActionType::Check
+		);
+	MenuBuilder.AddMenuEntry(
+		LOCTEXT("EditInstance", "Edit Instance"),
+		FText(),
+		FSlateIcon(FMeshEditorStyle::GetStyleSetName(), "MeshEditorMode.EditInstance"),
+		FUIAction
+		(
+			FExecuteAction::CreateLambda([this] { SetEditingPerInstance(IsEditingPerInstance() == false); }),
+			FCanExecuteAction::CreateLambda([this] {return true; }),
+			FIsActionChecked::CreateLambda([this] { return IsEditingPerInstance(); })
+		),
+		NAME_None,
+		EUserInterfaceActionType::Check
+		);
+
+	if (GetMeshElementSelectionMode() == EEditableMeshElementType::Polygon)
+	{
+		MenuBuilder.AddMenuEntry(
+			LOCTEXT("Move", "Move"),
+			FText(),
+			FSlateIcon(FMeshEditorStyle::GetStyleSetName(), "MeshEditorMode.PolyMove"),
+			FUIAction
+			(
+				FExecuteAction::CreateLambda([this] { EquippedPolygonAction = EMeshEditAction::Move; }),
+				FCanExecuteAction::CreateSP(this, &FMeshEditorMode::IsMeshElementTypeSelectedOrIsActiveSelectionMode, EEditableMeshElementType::Polygon),
+				FIsActionChecked::CreateLambda([this] { return (EquippedPolygonAction == EMeshEditAction::Move); })
+				),
+			NAME_None,
+			EUserInterfaceActionType::Check
+			);
+		MenuBuilder.AddMenuEntry(
+			LOCTEXT("Delete", "Delete"),
+			FText(),
+			FSlateIcon(FMeshEditorStyle::GetStyleSetName(), "MeshEditorMode.PolyDelete"),
+			FUIAction
+			(
+				FExecuteAction::CreateLambda([this] { DeleteSelectedMeshElement(); }),
+				FCanExecuteAction::CreateSP(this, &FMeshEditorMode::IsMeshElementTypeSelected, EEditableMeshElementType::Polygon)
+				),
+			NAME_None,
+			EUserInterfaceActionType::Button
+			);
+		MenuBuilder.AddMenuEntry(
+			LOCTEXT("Extrude", "Extrude"),
+			FText(),
+			FSlateIcon(FMeshEditorStyle::GetStyleSetName(), "MeshEditorMode.PolyExtrude"),
+			FUIAction
+			(
+				FExecuteAction::CreateLambda([this] { EquippedPolygonAction = EMeshEditAction::ExtrudePolygon; }),
+				FCanExecuteAction::CreateSP(this, &FMeshEditorMode::IsMeshElementTypeSelectedOrIsActiveSelectionMode, EEditableMeshElementType::Polygon),
+				FIsActionChecked::CreateLambda([this] { return (EquippedPolygonAction == EMeshEditAction::ExtrudePolygon); })
+			),
+			NAME_None,
+			EUserInterfaceActionType::Check
+			);
+		MenuBuilder.AddMenuEntry(
+			LOCTEXT("Inset", "Inset"),
+			FText(),
+			FSlateIcon(FMeshEditorStyle::GetStyleSetName(), "MeshEditorMode.PolyInset"),
+			FUIAction
+			(
+				FExecuteAction::CreateLambda([this] { EquippedPolygonAction = EMeshEditAction::InsetPolygon; }),
+				FCanExecuteAction::CreateSP(this, &FMeshEditorMode::IsMeshElementTypeSelectedOrIsActiveSelectionMode, EEditableMeshElementType::Polygon),
+				FIsActionChecked::CreateLambda([this] { return (EquippedPolygonAction == EMeshEditAction::InsetPolygon); })
+			),
+			NAME_None,
+			EUserInterfaceActionType::Check
+			);
+	}
+	else if (GetMeshElementSelectionMode() == EEditableMeshElementType::Edge)
+	{
+		MenuBuilder.AddMenuEntry(
+			LOCTEXT("Move", "Move"),
+			FText(),
+			FSlateIcon(FMeshEditorStyle::GetStyleSetName(), "MeshEditorMode.EdgeMove"),
+			FUIAction
+			(
+				FExecuteAction::CreateLambda([this] { EquippedEdgeAction = EMeshEditAction::Move; }),
+				FCanExecuteAction::CreateSP(this, &FMeshEditorMode::IsMeshElementTypeSelectedOrIsActiveSelectionMode, EEditableMeshElementType::Edge),
+				FIsActionChecked::CreateLambda([this] { return (EquippedEdgeAction == EMeshEditAction::Move); })
+			),
+			NAME_None,
+			EUserInterfaceActionType::Check
+			);
+		MenuBuilder.AddMenuEntry(
+			LOCTEXT("Delete", "Delete"),
+			FText(),
+			FSlateIcon(FMeshEditorStyle::GetStyleSetName(), "MeshEditorMode.EdgeDelete"),
+			FUIAction
+			(
+				FExecuteAction::CreateLambda([this] { DeleteSelectedMeshElement(); }),
+				FCanExecuteAction::CreateSP(this, &FMeshEditorMode::IsMeshElementTypeSelected, EEditableMeshElementType::Edge)
+				),
+			NAME_None,
+			EUserInterfaceActionType::Button
+			);
+		MenuBuilder.AddMenuEntry(
+			LOCTEXT("RemoveSelected", "Remove Selected"),
+			FText(),
+			FSlateIcon(FMeshEditorStyle::GetStyleSetName(), "MeshEditorMode.EdgeRemove"),
+			FUIAction
+			(
+				FExecuteAction::CreateLambda([this] { RemoveSelectedEdges(); }),
+				FCanExecuteAction::CreateSP(this, &FMeshEditorMode::IsMeshElementTypeSelected, EEditableMeshElementType::Edge)
+				),
+			NAME_None,
+			EUserInterfaceActionType::Button
+			);
+		MenuBuilder.AddMenuEntry(
+			LOCTEXT( "SelectEdgeLoop", "Select Edge Loop" ),
+			FText(),
+			FSlateIcon(FMeshEditorStyle::GetStyleSetName(), "MeshEditorMode.SelectLoop"),
+			FUIAction
+			(
+				FExecuteAction::CreateLambda([this] { SelectEdgeLoops(); }),
+				FCanExecuteAction::CreateSP(this, &FMeshEditorMode::IsMeshElementTypeSelected, EEditableMeshElementType::Edge)
+				),
+			NAME_None,
+			EUserInterfaceActionType::Button
+			);
+		MenuBuilder.AddMenuEntry(
+			LOCTEXT("InsertEdgeLoop", "Insert Loop"),
+			FText(),
+			FSlateIcon(FMeshEditorStyle::GetStyleSetName(), "MeshEditorMode.EdgeInsert"),
+			FUIAction
+			(
+				FExecuteAction::CreateLambda([this] { EquippedEdgeAction = EMeshEditAction::InsertEdgeLoop; }),
+				FCanExecuteAction::CreateSP(this, &FMeshEditorMode::IsMeshElementTypeSelectedOrIsActiveSelectionMode, EEditableMeshElementType::Edge),
+				FIsActionChecked::CreateLambda([this] { return (EquippedEdgeAction == EMeshEditAction::InsertEdgeLoop); })
+				),
+			NAME_None,
+			EUserInterfaceActionType::Check
+			);
+		MenuBuilder.AddMenuEntry(
+			LOCTEXT("ExtendEdge", "Extend"),
+			FText(),
+			FSlateIcon(FMeshEditorStyle::GetStyleSetName(), "MeshEditorMode.EdgeExtend"),
+			FUIAction
+			(
+				FExecuteAction::CreateLambda([this] { EquippedEdgeAction = EMeshEditAction::ExtendEdge; }),
+				FCanExecuteAction::CreateSP(this, &FMeshEditorMode::IsMeshElementTypeSelectedOrIsActiveSelectionMode, EEditableMeshElementType::Edge),
+				FIsActionChecked::CreateLambda([this] { return (EquippedEdgeAction == EMeshEditAction::ExtendEdge); })
+				),
+			NAME_None,
+			EUserInterfaceActionType::Check
+			);
+	}
+	else if (GetMeshElementSelectionMode() == EEditableMeshElementType::Vertex)
+	{
+		MenuBuilder.AddMenuEntry(
+			LOCTEXT("Move", "Move"),
+			FText(),
+			FSlateIcon(FMeshEditorStyle::GetStyleSetName(), "MeshEditorMode.VertexMove"),
+			FUIAction
+			(
+				FExecuteAction::CreateLambda([this] { EquippedVertexAction = EMeshEditAction::Move; }),
+				FCanExecuteAction::CreateSP(this, &FMeshEditorMode::IsMeshElementTypeSelectedOrIsActiveSelectionMode, EEditableMeshElementType::Vertex),
+				FIsActionChecked::CreateLambda([this] { return (EquippedVertexAction == EMeshEditAction::Move); })
+				),
+			NAME_None,
+			EUserInterfaceActionType::Check
+			);
+		MenuBuilder.AddMenuEntry(
+			LOCTEXT("Delete", "Delete"),
+			FText(),
+			FSlateIcon(FMeshEditorStyle::GetStyleSetName(), "MeshEditorMode.VertexDelete"),
+			FUIAction
+			(
+				FExecuteAction::CreateLambda([this] { DeleteSelectedMeshElement(); }),
+				FCanExecuteAction::CreateSP(this, &FMeshEditorMode::IsMeshElementTypeSelected, EEditableMeshElementType::Vertex)
+			),
+			NAME_None,
+			EUserInterfaceActionType::Button
+			);
+		MenuBuilder.AddMenuEntry(
+			LOCTEXT("RemoveSelected", "Remove Selected"),
+			FText(),
+			FSlateIcon(FMeshEditorStyle::GetStyleSetName(), "MeshEditorMode.VertexRemove"),
+			FUIAction
+			(
+				FExecuteAction::CreateLambda([this] { RemoveSelectedVertices(); }),
+				FCanExecuteAction::CreateSP(this, &FMeshEditorMode::IsMeshElementTypeSelected, EEditableMeshElementType::Vertex)
+			),
+			NAME_None,
+			EUserInterfaceActionType::Button
+			);
+		MenuBuilder.AddMenuEntry(
+			LOCTEXT("Extend", "Extend"),
+			FText(),
+			FSlateIcon(FMeshEditorStyle::GetStyleSetName(), "MeshEditorMode.VertexExtend"),
+			FUIAction
+			(
+				FExecuteAction::CreateLambda([this] { EquippedVertexAction = EMeshEditAction::ExtendVertex; }),
+				FCanExecuteAction::CreateSP(this, &FMeshEditorMode::IsMeshElementTypeSelectedOrIsActiveSelectionMode, EEditableMeshElementType::Vertex),
+				FIsActionChecked::CreateLambda([this] { return (EquippedVertexAction == EMeshEditAction::ExtendVertex); })
+				),
+			NAME_None,
+			EUserInterfaceActionType::Check
+			);
+		MenuBuilder.AddMenuEntry(
+			LOCTEXT("WeldSelected", "Weld Selected"),
+			FText(),
+			FSlateIcon(FMeshEditorStyle::GetStyleSetName(), "MeshEditorMode.VertexWeld"),
+			FUIAction
+			(
+				FExecuteAction::CreateLambda([this] { WeldSelectedVertices(); }),
+				FCanExecuteAction::CreateSP(this, &FMeshEditorMode::IsMeshElementTypeSelected, EEditableMeshElementType::Vertex)
+				),
+			NAME_None,
+			EUserInterfaceActionType::Button
+			);
+	}
+}
+
+#undef LOCTEXT_NAMESPACE
