Branch snapshot for CL 3327560

diff --git a/Engine/Source/Editor/MeshEditor/MeshEditorMode.h b/Engine/Source/Editor/MeshEditor/MeshEditorMode.h
new file mode 100644
--- /dev/null
+++ b/Engine/Source/Editor/MeshEditor/MeshEditorMode.h
@@ -0,0 +1,717 @@
+// Copyright 1998-2015 Epic Games, Inc. All Rights Reserved.
+
+#pragma once
+
+#include "EditableMeshTypes.h"
+#include "EditableMesh.h"
+#include "EdMode.h"
+#include "MeshElement.h"
+#include "ArrayView.h"
+#include "LevelEditorViewport.h"	// For FTrackingTransaction
+#include "UnrealEdMisc.h"
+#include "IMeshEditorMode.h"
+#include "MeshEditorMode.generated.h"
+
+
+UCLASS()
+class UMeshEditorModeProxyObject : public UObject
+{
+	GENERATED_BODY()
+
+public:
+
+	/** The mesh editor that owns us */
+	class FMeshEditorMode* OwningMeshEditorMode;
+};
+
+
+/**
+ * Mesh Editor Mode.  Extends editor viewports with the ability to edit meshes
+ */
+class FMeshEditorMode : public FEdMode, public IMeshEditorMode
+{
+public:
+
+	/**
+	 * The types of interactor shapes we support
+	 */
+	enum class EInteractorShape
+	{
+		/** Invalid shape (or none) */
+		Invalid,
+
+		/** Grabber sphere */
+		GrabberSphere,
+
+		/** Laser pointer shape */
+		Laser,
+	};
+
+	/**
+	 * Contains state for either a mouse cursor or a virtual hand (in VR), to be used to interact with a mesh
+	 */
+	struct FMeshEditorInteractorData
+	{
+		/** The viewport interactor that is this data's counterpart */
+		TWeakObjectPtr<class UViewportInteractor> ViewportInteractor;
+
+		/** True if we have a valid interaction grabber sphere right now */
+		bool bGrabberSphereIsValid;
+
+		/** The sphere for radial interactions */
+		FSphere GrabberSphere;
+
+		/** True if we have a valid interaction ray right now */
+		bool bLaserIsValid;
+
+		/** World space start location of the interaction ray the last time we were ticked */
+		FVector LaserStart;
+
+		/** World space end location of the interaction ray */
+		FVector LaserEnd;
+
+		/** What shape of interactor are we using to hover? */
+		EInteractorShape HoverInteractorShape;
+
+		/** Information about a mesh we're hovering over or editing */
+		FMeshElement HoveredMeshElement;
+
+		/** The element we were hovering over last frame */
+		FMeshElement PreviouslyHoveredMeshElement;
+
+		/** The hover point.  With a ray, this could be the impact point along the ray.  With grabber sphere interaction, this 
+		    would be the point within the sphere radius where we've found a point on an object to interact with */
+		FVector HoverLocation;
+
+
+		/** Default constructor that initializes everything to safe values */
+		FMeshEditorInteractorData();
+	};
+
+
+public:
+
+	/** Default constructor for FMeshEditorMode */
+	FMeshEditorMode();
+
+	/** Cleans up this mode, called when the editor is shutting down */
+	virtual ~FMeshEditorMode();
+
+	// FEdMode interface
+	virtual void Enter() override;
+	virtual void Exit() override;
+	virtual void Tick( FEditorViewportClient* ViewportClient, float DeltaTime ) override;
+	virtual bool InputKey( FEditorViewportClient* ViewportClient, FViewport* Viewport, FKey Key, EInputEvent Event ) override;
+	virtual bool InputAxis( FEditorViewportClient* ViewportClient, FViewport* Viewport, int32 ControllerId, FKey Key, float Delta, float DeltaTime ) override;
+	virtual bool InputDelta( FEditorViewportClient* ViewportClient, FViewport* Viewport, FVector& Drag, FRotator& Rotation, FVector& Scale ) override;
+	virtual bool IsCompatibleWith( FEditorModeID OtherModeID ) const override;
+	virtual void AddReferencedObjects( FReferenceCollector& Collector ) override;
+	virtual void Render( const FSceneView* SceneView, FViewport* Viewport, FPrimitiveDrawInterface* PDI ) override;	
+	virtual void PostUndo() override;
+	virtual bool FrustumSelect( const FConvexVolume& InFrustum, FEditorViewportClient* InViewportClient, bool InSelect = true ) override;
+	virtual bool ShouldDrawWidget() const override;
+
+	// IMeshEditorMode Interface
+	virtual EEditableMeshElementType GetMeshElementSelectionMode() const override{ return MeshElementSelectionMode; }
+	virtual void SetMeshElementSelectionMode( EEditableMeshElementType ElementType ) override;
+	virtual void RegisterWithExternalMenuSystem() override;
+	virtual void UnregisterWithExternalMenuSystem() override;
+	// End IMeshEditorMode Interface
+
+
+	/** Returns the type of elements that are selected right now, or Invalid if nothing is selected */
+	EEditableMeshElementType GetSelectedMeshElementType() const;
+
+	/** Returns whether the specified element type is selected */
+	bool IsMeshElementTypeSelected( EEditableMeshElementType ElementType ) const { return GetSelectedMeshElementType() == ElementType; }
+
+	/** Returns whether either the specified element type is selected, or we're in the selection mode for that element type */
+	bool IsMeshElementTypeSelectedOrIsActiveSelectionMode( EEditableMeshElementType ElementType ) const { return GetMeshElementSelectionMode() == ElementType || GetSelectedMeshElementType() == ElementType; }
+
+	const TArray<TTuple<TSharedPtr<FUICommandInfo>, FUIAction>>& GetVertexActions() const { return VertexActions; }
+	const TArray<TTuple<TSharedPtr<FUICommandInfo>, FUIAction>>& GetEdgeActions() const { return EdgeActions; }
+	const TArray<TTuple<TSharedPtr<FUICommandInfo>, FUIAction>>& GetPolygonActions() const { return PolygonActions; }
+
+	/** Checks to see that the mesh element actually exists in the mesh */
+	inline static bool IsElementIDValid( const FMeshElement& MeshElement, const UEditableMesh* EditableMesh )
+	{
+		bool bIsValid = false;
+
+		if( MeshElement.ElementAddress.ElementID != FElementID::Invalid )
+		{
+			switch( MeshElement.ElementAddress.ElementType )
+			{
+				case EEditableMeshElementType::Vertex:
+					bIsValid = EditableMesh->IsValidVertex( FVertexID( MeshElement.ElementAddress.ElementID ) );
+					break;
+
+				case EEditableMeshElementType::Edge:
+					bIsValid = EditableMesh->IsValidEdge( FEdgeID( MeshElement.ElementAddress.ElementID ) );
+					break;
+
+				case EEditableMeshElementType::Polygon:
+					bIsValid = EditableMesh->IsValidPolygon( FPolygonRef( MeshElement.ElementAddress.SectionID, FPolygonID( MeshElement.ElementAddress.ElementID ) ) );
+					break;
+			}
+		}
+
+		return bIsValid;
+	}
+
+	bool IsEditingPerInstance() const { return bPerInstanceEdits; }
+	void SetEditingPerInstance( bool bPerInstance ) { bPerInstanceEdits = bPerInstance; }
+
+	/** Propagates instance changes to the static mesh asset */
+	void PropagateInstanceChanges();
+
+	/** Whether there are instance changes which can be propagated */
+	bool CanPropagateInstanceChanges() const;
+
+	/** Gets the container of all the assets used in the mesh editor */
+	const class UMeshEditorAssetContainer& GetAssetContainer() const;
+
+protected:
+
+	/** Gets an editable mesh from our cache of editable meshes for the specified sub-mesh address, or tries to create and cache a new 
+	    editable mesh if we haven't seen this sub-mesh address before.  Can return nullptr if no mesh was possible for that address. */
+	UEditableMesh* FindOrCreateEditableMesh( class UPrimitiveComponent& Component, const FEditableMeshSubMeshAddress& SubMeshAddress );
+
+	/** Fills the specified dynamic mesh builder with primitives to render a mesh vertex */
+	void AddVertexToDynamicMesh( const UEditableMesh& EditableMesh, const FTransform& CameraToWorld, const FMatrix& ComponentToWorldMatrix, const FVertexID VertexID, const FColor ColorAndOpacity, const float SizeBias, const bool bApplyDepthBias, class FDynamicMeshBuilder& MeshBuilder );
+		
+	/** Fills the specified dynamic mesh builder with primitives to a thick line.  Incoming positions are in world space. */
+	void AddThickLineToDynamicMesh( const FTransform& CameraToWorld, const FVector EdgeVertexPositions[2], const FColor ColorAndOpacity, const float SizeBias, const bool bApplyDepthBias, FDynamicMeshBuilder& MeshBuilder );
+
+	/** Fills the specified dynamic mesh builder with primitives to render a mesh edge */
+	void AddEdgeToDynamicMesh( const UEditableMesh& EditableMesh, const FTransform& CameraToWorld, const FMatrix& ComponentToWorldMatrix, const FEdgeID EdgeID, const FColor ColorAndOpacity, const float SizeBias, class FDynamicMeshBuilder& MeshBuilder );
+
+	/** Fills the specified dynamic mesh builders with primitives to render a polygon and it's edges */
+	void AddPolygonToDynamicMesh( const UEditableMesh& EditableMesh, const FTransform& CameraToWorld, const FMatrix& ComponentToWorldMatrix, const FPolygonRef PolygonRef, const FColor ColorAndOpacity, const float SizeBias, const bool bFillFaces, class FDynamicMeshBuilder& VertexAndEdgeMeshBuilder, class FDynamicMeshBuilder* PolygonFaceMeshBuilder );
+
+	/** Renders the spceified mesh element */
+	void DrawMeshElements( const FTransform& CameraToWorld, FViewport* Viewport, FPrimitiveDrawInterface* PDI, const TArrayView<FMeshElement>& MeshElements, const FColor Color, const bool bFillFaces, const float HoverAnimation, const TArray<FColor>* OptionalPerElementColors = nullptr, const TArray<float>* OptionalPerElementSizeBiases = nullptr );
+
+	/** Called every frame for each viewport interactor to update what's under the cursor */
+	void OnViewportInteractionHoverUpdate( class UViewportInteractor* ViewportInteractor, FVector& OutHoverImpactPoint, bool& bWasHandled );
+		
+	/** Called when the user presses a button on their mouse or motion controller device */
+	void OnViewportInteractionInputAction( class FEditorViewportClient& ViewportClient, class UViewportInteractor* ViewportInteractor, const struct FViewportActionKeyInput& Action, bool& bOutIsInputCaptured, bool& bWasHandled );
+
+	/** Called when the user clicks on the background */
+	void OnViewportInteractionInputUnhandled( class FEditorViewportClient& ViewportClient, class UViewportInteractor* ViewportInteractor, const struct FViewportActionKeyInput& Action );
+
+	/** Called when the viewport interaction system starts dragging transformable objects around */
+	void OnViewportInteractionStartDragging( class UViewportInteractor* ViewportInteractor );
+
+	/** Called when the viewport interaction system stops dragging transformable objects around */
+	void OnViewportInteractionStopDragging( class UViewportInteractor* ViewportInteractor );
+
+	/** Called when the viewport interaction finishes moving a set of transformable objects */
+	void OnViewportInteractionFinishedMovingTransformables();
+	
+	/** Called when VR editor world interaction drags a material onto a component */
+	void OnVREditorModePlaceDraggedMaterial( class UPrimitiveComponent* HitComponent, class UMaterialInterface* MaterialInterface, bool& bPlaced );
+
+	/** Applies a modification to the mesh that's currently hovered */
+	void UpdateActiveAction( const bool bIsActionFinishing );
+
+	/** Geometry tests */
+	FEditableMeshElementAddress QueryElement( const UEditableMesh& EditableMesh, const bool bUseSphere, const FSphere& Sphere, const float SphereFuzzyDistance, bool bUseRay, const FVector& RayStart, const FVector& RayEnd, const float RayFuzzyDistance, const EEditableMeshElementType OnlyElementType, const FVector& CameraLocation, const float FuzzyDistanceScaleFactor, EInteractorShape& OutInteractorShape, FVector& OutHitLocation ) const;
+	static bool CheckVertex( const bool bUseSphere, const FSphere& Sphere, const float SphereFuzzyDistance, bool bUseRay, const FVector& RayStart, const FVector& RayEnd, const float FuzzyDistance, const FVector& VertexPosition, const FVector& CameraLocation, const float FuzzyDistanceScaleFactor, EInteractorShape& ClosestInteractorShape, float& ClosestDistanceToRay, float& ClosestDistanceOnRay, FVector& ClosestHitLocation, const bool bAlreadyHitVertex );
+	static bool CheckEdge( const bool bUseSphere, const FSphere& Sphere, const float SphereFuzzyDistance, bool bUseRay, const FVector& RayStart, const FVector& RayEnd, const float FuzzyDistance, const FVector EdgeVertexPositions[ 2 ], const FVector& CameraLocation, const float FuzzyDistanceScaleFactor, EInteractorShape& ClosestInteractorShape, float& ClosestDistanceToRay, float& ClosestDistanceOnRay, FVector& ClosestHitLocation, const bool bAlreadyEdge );
+	static bool CheckTriangle( const bool bUseSphere, const FSphere& Sphere, const float SphereFuzzyDistance, bool bUseRay, const FVector& RayStart, const FVector& RayEnd, const float FuzzyDistance, const FVector TriangleVertexPositions[ 3 ], const FVector& CameraLocation, const float FuzzyDistanceScaleFactor, EInteractorShape& ClosestInteractorShape, float& ClosestDistanceToRay, float& ClosestDistanceOnRay, FVector& ClosestHitLocation, const bool bAlreadyHitTriangle );
+
+	/** Returns the index of an element in the selection set, or INDEX_NONE if its not selected */
+	int32 GetSelectedMeshElementIndex( const FMeshElement& MeshElement ) const;
+
+	/** Clears hover and selection on mesh elements that may no longer be valid.  You'll want to call this if you change the mesh topology */
+	TUniquePtr<FChange> ClearInvalidSelectedElements_Internal();
+
+	/** Clears hover on mesh elements that may no longer be valid.  You'll want to call this if you change the mesh topology */
+	void ClearInvalidHoveredElements();
+
+	/** Updates the current view location, from either the viewport client or the VR interface, whichever is in use */
+	void UpdateCameraToWorldTransform( const FEditorViewportClient& ViewportClient );
+
+	enum class EMeshEditAction
+	{
+		/** Nothing going on right now */
+		None,
+
+		/** Selecting mesh elements by 'painting' over multiple elements */
+		SelectByPainting,
+
+		/** Moving elements using a transform gizmo */
+		MoveUsingGizmo,
+
+		/** Moving selected mesh elements (vertices, edges or polygons) */
+		Move,
+
+		/** Split an edge by inserting a vertex.  You can drag to preview where the vertex will be inserted. */
+		SplitEdge,
+
+		/** Splits an edge by inserting a new vertex, then immediately starts dragging that vertex */
+		SplitEdgeAndDragVertex,
+
+		/** Insert an edge loop */
+		InsertEdgeLoop,
+
+		/** Extrude polygon by making a copy of it and allowing you to shift it along the polygon normal axis */
+		ExtrudePolygon,
+
+		/** Extrude polygon by making a copy of it and allowing you to move it around freely */
+		FreelyExtrudePolygon,
+
+		/** Inset polygon by replacing it with a new polygon that is bordered by polygons of a specific relative size */
+		InsetPolygon,
+
+		/** Bevel polygons by adding angled bordering polygons of a specific relative size */
+		BevelPolygon,
+
+		/** Extend an edge by making a copy of it and allowing you to move it around */
+		ExtendEdge,
+
+		/** Extend a vertex by making a copy of it, creating new polygons to join the geometry together */
+		ExtendVertex,
+
+		/** For subdivision meshes, edits how sharp a vertex corner is by dragging in space */
+		EditVertexCornerSharpness,
+
+		/** For subdivision meshes, edits how sharp an edge crease is by dragging in space */
+		EditEdgeCreaseSharpness,
+
+		/** Freehand vertex drawing */
+		DrawVertices,
+	};
+
+	/** Begins an action */
+	void StartAction( const EMeshEditAction NewAction, class UViewportInteractor* ActionInteractor, const bool bActionNeedsHoverLocation, const FText& UndoText );
+
+	/** Ends an action that's currently in progress.  Usually called when the user commits a change by clicking/releasing, but can
+	    also be called when the user begins a new action while inertia is still influencing the active action */
+	void FinishAction();
+
+	/** @return Returns true if the undo system is available right now.  When in Simulate Mode, we can't store undo states or use undo/redo features */
+	static bool IsUndoSystemAvailable();
+
+	/** Saves undo state, if possible (e.g., not in Simulate mode.) */
+	static void StoreUndo( UObject* Object, TUniquePtr<FChange> UndoChange );
+
+	/** Binds UI commands to actions for the mesh editor */
+	void BindCommands();
+
+	void RegisterCommonEditingMode( const TSharedPtr<FUICommandInfo>& Command, EMeshEditAction EditingMode );
+	void RegisterVertexEditingMode( const TSharedPtr<FUICommandInfo>& Command, EMeshEditAction EditingMode );
+	void RegisterEdgeEditingMode( const TSharedPtr<FUICommandInfo>& Command, EMeshEditAction EditingMode );
+	void RegisterPolygonEditingMode( const TSharedPtr<FUICommandInfo>& Command, EMeshEditAction EditingMode );
+
+	void RegisterCommand( const TSharedPtr<FUICommandInfo>& Command, const FExecuteAction& ExecuteAction );
+	void RegisterVertexCommand( const TSharedPtr<FUICommandInfo>& Command, const FExecuteAction& ExecuteAction );
+	void RegisterEdgeCommand( const TSharedPtr<FUICommandInfo>& Command, const FExecuteAction& ExecuteAction );
+	void RegisterPolygonCommand( const TSharedPtr<FUICommandInfo>& Command, const FExecuteAction& ExecuteAction );
+
+	/** Return the CommandList pertinent to the currently selected element type, or nullptr if nothing is selected */
+	const FUICommandList* GetCommandListForSelectedElementType() const;
+
+	/** Commits the mesh instance for the given component */
+	void CommitEditableMeshIfNecessary( UEditableMesh* EditableMesh, UPrimitiveComponent* Component );
+
+	/** Commits all selected meshes */
+	void CommitSelectedMeshes();
+
+	/** Deletes selected polygons, or polygons partly defined by selected elements; returns whether successful */
+	bool DeleteSelectedMeshElement();
+
+	/** Tessellates selected polygons into smaller polygons; returns whether it did anything */
+	bool TessellateSelectedPolygons();
+
+	/** Adds or removes a subdivision level for selected meshes */
+	void AddOrRemoveSubdivisionLevel( const bool bShouldAdd );
+
+	/** Quadrangulates the currently selected mesh */
+	void QuadrangulateMesh();
+
+	/** Moves the viewport camera to frame the currently selected elements */
+	void FrameSelectedElements( FEditorViewportClient* ViewportClient );
+
+	/** Removes the selected edge if possible; returns whether successful */
+	bool RemoveSelectedEdges();
+
+	/** Selects the edge loops which contain the selected edges */
+	bool SelectEdgeLoops();
+
+	/** Removes the selected vertex if possible; returns whether successful */
+	bool RemoveSelectedVertices();
+
+	/** Welds the selected vertices if possible, keeping the first selected vertex */
+	bool WeldSelectedVertices();
+
+	/** Flips selected polygons; returns whether successful */
+	bool FlipSelectedPolygons();
+
+	/** Triangulates selected polygons; returns whether successful */
+	bool TriangulateSelectedPolygons();
+	
+	/** Assigns a material to the selected polygons; returns whether successful */
+	bool AssignSelectedMaterialToSelectedPolygons();
+
+	/** Assigns a material to the selected polygons; returns whether successful */
+	bool AssignMaterialToSelectedPolygons( UMaterialInterface* SelectedMaterial );
+
+	/** Creases selected edges; returns whether successful */
+	bool MakeSelectedEdgesHardOrSoft( bool bMakeEdgesHard );
+
+	/** Rolls back whatever we changed last time while previewing */
+	void RollbackPreviewChanges();
+
+	/** Gets mesh editor interactor data for the specified viewport interactor.  If we've never seen this viewport interactor before,
+	    new (empty) data will be created for it on demand */
+	FMeshEditorInteractorData& GetMeshEditorInteractorData( UViewportInteractor* ViewportInteractor );
+
+	/** Deselects the mesh elements specified */
+	void DeselectMeshElements( const TMap<UEditableMesh*, TArray<FMeshElement>>& MeshElementsToDeselect );
+
+	/** Helper function that returns a map keying an editable mesh with its selected elements */
+	void GetSelectedMeshesAndElements( EEditableMeshElementType ElementType, TMap<UEditableMesh*, TArray<FMeshElement>>& OutMeshesAndElements );
+
+	void GetSelectedMeshesAndVertices( TMap<UEditableMesh*, TArray<FMeshElement>>& OutMeshesAndVertices ) { GetSelectedMeshesAndElements( EEditableMeshElementType::Vertex, OutMeshesAndVertices ); }
+	void GetSelectedMeshesAndEdges( TMap<UEditableMesh*, TArray<FMeshElement>>& OutMeshesAndEdges ) { GetSelectedMeshesAndElements( EEditableMeshElementType::Edge, OutMeshesAndEdges ); }
+	void GetSelectedMeshesAndPolygons( TMap<UEditableMesh*, TArray<FMeshElement>>& OutMeshesAndPolygons ) { GetSelectedMeshesAndElements( EEditableMeshElementType::Polygon, OutMeshesAndPolygons ); }
+
+	/** Helper function that converts selected vertices into a map of meshes to their respective selected vertex elements and vertex IDs */
+	void GetSelectedMeshesAndVertices( TMap< UEditableMesh*, TArray< TTuple< FMeshElement, FVertexID > > >& OutMeshesAndVertices );
+
+	/** Helper function that converts selected edges into a map of meshes to their respective selected edge elements and edge IDs */
+	void GetSelectedMeshesAndEdges( TMap< UEditableMesh*, TArray< TTuple< FMeshElement, FEdgeID > > >& OutMeshesAndEdges );
+
+	/** Helper function that converts selected polygons into a map of meshes to their respective selected polygon elements and polygon refs */
+	void GetSelectedMeshesAndPolygons( TMap< UEditableMesh*, TArray< TTuple< FMeshElement, FPolygonRef > > >& OutMeshesAndPolygons );
+
+	/** Given an interactor and a mesh, finds edges under the interactor along with their exact split position (progress along the edge) */
+	void FindEdgeSplitUnderInteractor( UViewportInteractor* ViewportInteractor, const UEditableMesh* EditableMesh, const TArray<TTuple<FMeshElement, FEdgeID>>& Edges, TArray<float>& OutSplits );
+
+	/** Selects elements of the given type captured by the last marquee select */
+	void PerformMarqueeSelect( EEditableMeshElementType ElementType );
+
+	/** Rebuilds the list of mesh element transformables and updates the world viewport interaction system with the new list */
+	void RefreshTransformables();
+
+	/** Callback when PIE/SIE ends */
+	void OnEndPIE( bool bIsSimulating );
+
+	/** Callback from the level editor when the map changes */
+	void OnMapChanged( UWorld* World, EMapChangeType MapChangeType );
+
+	/** Callback from the level editor when new actors become selected or deselected */
+	void OnActorSelectionChanged( const TArray<UObject*>& NewSelection, bool bForceRefresh );
+
+	/** Creates the mesh edit actions to pass to the radial menu generator */
+	void MeshEditActionsGenerator(class FMenuBuilder MenuBuilder, TSharedPtr<class FUICommandList> CommandList, class UVREditorMode* VRMode, float& RadiusOverride);
+
+	/** Clears any references to editable meshes (which may now be invalid) */
+	void RemoveEditableMeshReferences();
+
+	/** Returns whether the mode is currently active */
+	bool IsActive() const { return ( ViewportWorldInteraction != nullptr ); }
+
+	/** Plays sound when starting a mesh edit action */
+	void PlayStartActionSound(EMeshEditAction NewAction, UViewportInteractor* ActionInteractor = nullptr);
+
+	/** Plays sound when mesh edit action was finished */
+	void PlayFinishActionSound(EMeshEditAction NewAction, UViewportInteractor* ActionInteractor = nullptr);
+
+protected:
+
+	struct FSelectOrDeselectMeshElementsChangeInput
+	{
+		/** New mesh elements that should become selected */
+		TArray<FMeshElement> MeshElementsToSelect;
+
+		/** Mesh elements that should be deselected */
+		TArray<FMeshElement> MeshElementsToDeselect;
+
+		/** Default constructor */
+		FSelectOrDeselectMeshElementsChangeInput()
+			: MeshElementsToSelect(),
+			  MeshElementsToDeselect()
+		{
+		}
+	};
+
+
+	class FSelectOrDeselectMeshElementsChange : public FChange
+	{
+
+	public:
+
+		/** Constructor */
+		FSelectOrDeselectMeshElementsChange( const FSelectOrDeselectMeshElementsChangeInput& InitInput )
+			: Input( InitInput )
+		{
+		}
+
+		FSelectOrDeselectMeshElementsChange( FSelectOrDeselectMeshElementsChangeInput&& InitInput )
+			: Input( MoveTemp( InitInput ) )
+		{
+		}
+
+		// Parent class overrides
+		virtual TUniquePtr<FChange> Execute( UObject* Object ) override;
+		virtual FString ToString() const override;
+
+	private:
+
+		/** The data we need to make this change */
+		FSelectOrDeselectMeshElementsChangeInput Input;
+
+	};
+
+	struct FDeselectAllMeshElementsChangeInput
+	{
+		/** Default constructor */
+		FDeselectAllMeshElementsChangeInput()
+		{
+		}
+	};
+
+
+	class FDeselectAllMeshElementsChange : public FChange
+	{
+
+	public:
+
+		/** Constructor */
+		FDeselectAllMeshElementsChange( const FDeselectAllMeshElementsChangeInput& InitInput )
+			: Input( InitInput )
+		{
+		}
+
+		// Parent class overrides
+		virtual TUniquePtr<FChange> Execute( UObject* Object ) override;
+		virtual FString ToString() const override;
+
+	private:
+
+		/** The data we need to make this change */
+		FDeselectAllMeshElementsChangeInput Input;
+
+	};
+
+
+	struct FSetElementSelectionModeChangeInput
+	{
+		/** Default constructor */
+		FSetElementSelectionModeChangeInput()
+			: Mode( EEditableMeshElementType::Invalid )
+			, bApplyStoredSelection( false )
+		{
+		}
+
+		/** The mesh element selection mode being set */
+		EEditableMeshElementType Mode;
+
+		/** Whether we should just apply the stored selection */
+		bool bApplyStoredSelection;
+
+		/** The stored selection to be optionally applied */
+		TArray<FMeshElement> StoredSelection;
+	};
+
+
+	class FSetElementSelectionModeChange : public FChange
+	{
+
+	public:
+
+		/** Constructor */
+		FSetElementSelectionModeChange( const FSetElementSelectionModeChangeInput& InitInput )
+			: Input( InitInput )
+		{
+		}
+
+		// Parent class overrides
+		virtual TUniquePtr<FChange> Execute( UObject* Object ) override;
+		virtual FString ToString() const override;
+
+	private:
+
+		/** The data we need to make this change */
+		FSetElementSelectionModeChangeInput Input;
+	};
+
+
+	/** Cached pointer to the viewport world interaction object we're using to interact with mesh elements */
+	class UViewportWorldInteraction* ViewportWorldInteraction;
+
+	/** Cached pointer to the VREditorMode object we're using */
+	class UVREditorMode* VREditorMode;
+
+	/** Material to use to render hovered mesh geometry */
+	UMaterialInterface* HoveredGeometryMaterial;
+
+	/** Material to use to render hovered triangles or faces */
+	UMaterialInterface* HoveredFaceMaterial;
+
+	/** Hover feedback animation time value, ever-incrementing until selection changes */
+	double HoverFeedbackTimeValue;
+
+	/** Interactors for the mouse cursor, and also for either virtual hand (when using VR) */
+	TArray<FMeshEditorInteractorData> MeshEditorInteractorDatas;
+
+	/** Specifies the type of element which is currently being selected */
+	EEditableMeshElementType MeshElementSelectionMode;
+
+	/** List of mesh elements that we've selected.  All elements in this list will always have the same mesh element type.
+	    We don't allow users to select both edges, faces and/or polygons at the same time. */
+	TArray<FMeshElement> SelectedMeshElements;
+
+	/** List of selected elements for different selection modes. */
+	TArray<FMeshElement> SelectedVertices;
+	TArray<FMeshElement> SelectedEdges;
+	TArray<FMeshElement> SelectedPolygons;
+
+	/** List of old hovered mesh elements that are in the process of being faded out */
+	TArray<FMeshElement> FadingOutHoveredMeshElements;
+
+
+
+
+	/** Cached editable meshes */
+	// @todo mesheditor: Need to expire these at some point, otherwise we just grow and grow */
+	TMap< FEditableMeshSubMeshAddress, class UEditableMesh* > CachedEditableMeshes;
+
+	/** Manages saving undo for selected mesh elements while we're dragging them around */
+	FTrackingTransaction TrackingTransaction;
+
+	/** The next action that will be started when interacting with a selected vertex */
+	EMeshEditAction EquippedVertexAction;
+
+	/** The next action that will be started when interacting with a selected edge */
+	EMeshEditAction EquippedEdgeAction;
+
+	/** The next action that will be started when interacting with a selected polygon */
+	EMeshEditAction EquippedPolygonAction;
+
+	/** The interactive action currently being performed (and previewed).  These usually happen over multiple frames, and
+	    result in a 'final' application of the change that performs a more exhaustive (and more expensive) update. */
+	EMeshEditAction ActiveAction;
+
+	/** When performing an interactive action that was initiated using an interactor, this is the interactor that was used. */
+	class UViewportInteractor* ActiveActionInteractor;
+
+	/** True if the ActiveAction needs us to update the mesh element under the cursor every frame */
+	bool bActiveActionNeedsHoverLocation;
+
+	/** True if UpdateActiveAction() has yet to be called since the current action started */
+	bool bIsFirstActiveActionUpdate;
+
+	/** Command list for actions available regardless of selection */
+	TSharedPtr<FUICommandList> CommonCommands;
+
+	/** Command list for actions available when a vertex is selected */
+	TSharedPtr<FUICommandList> VertexCommands;
+
+	/** Command list for actions available when an edge is selected */
+	TSharedPtr<FUICommandList> EdgeCommands;
+
+	/** Command list for actions available when a polygon is selected */
+	TSharedPtr<FUICommandList> PolygonCommands;
+
+	TArray<TTuple<TSharedPtr<FUICommandInfo>, FUIAction>> CommonActions;
+	TArray<TTuple<TSharedPtr<FUICommandInfo>, FUIAction>> VertexActions;
+	TArray<TTuple<TSharedPtr<FUICommandInfo>, FUIAction>> EdgeActions;
+	TArray<TTuple<TSharedPtr<FUICommandInfo>, FUIAction>> PolygonActions;
+
+	//
+	// ExtrudePolygon
+	//
+
+	/** When extruding polygons, we need to keep track of the interactor's impact point and normal, because
+	    the user is going to be aiming their interactor along that axis to choose an extrusion point */
+	FVector ExtrudePolygonAxisOrigin;
+	FVector ExtrudePolygonAxisDirection;
+
+
+	//
+	// InsetPolygon
+	//
+
+	/** The selected polygon we clicked on to start the inset action */
+	FMeshElement InsetUsingPolygonElement;
+
+
+	//
+	// SplitEdgeAndDragVertex
+	//
+
+	/** When splitting an edge and dragging a vertex, this is the list of edges that will be split */
+	TMap< UEditableMesh*, TArray< TTuple< FMeshElement, FEdgeID > > > SplitEdgeMeshesAndEdgesToSplit;
+
+	/** When splitting an edge and dragging a vertex, this is the list of split positions along those edges */
+	TArray<float> SplitEdgeSplitList;
+
+
+	//
+	// EditVertexCornerSharpness/EdgeEdgeCreaseSharpness
+	//
+
+	/** Where the active interactor's impact point was when the "edit sharpness" action started */
+	FVector EditSharpnessStartLocation;
+
+	
+	//
+	// DrawVertices
+	//
+
+	/** Array of all points drawn so far */
+	TArray<TTuple<double, FVector>> DrawnPoints;
+
+
+	// ...
+
+	/** When interactively dragging to preview a change (that might not be fully committed), this is the Change that
+	    will be used to roll back the previewed alternative from the previous frame */
+	TArray<TTuple<class UObject*,TUniquePtr<FChange>>> PreviewRevertChanges;
+
+	/** Whether topology changed while we were applying the preview changes that we might revert later */
+	EMeshTopologyChange PreviewTopologyChange;
+
+	/** Proxy UObject to pass to the undo system when performing interactions that affect state of the mode itself, 
+	    such as the selection set.  We need this because the UE4 undo system requries a UObject, but we're an FEdMode. */
+	UMeshEditorModeProxyObject* MeshEditorModeProxyObject;
+
+	/** When selecting by painting (EMeshEditAction::SelectByPainting), this is the compound change that can be applied to
+	    roll back the change to select.  We'll build this up as the user is painting select, then store it in the undo buffer */
+	TUniquePtr<FCompoundChangeInput> SelectingByPaintingRevertChangeInput;
+
+	/** Whether vertex normals should be displayed for the selected mesh */
+	bool bShowVertexNormals;
+
+	/** Results of marquee select operation, pending action */
+	TArray<FMeshElement> MarqueeSelectVertices;
+	TArray<FMeshElement> MarqueeSelectEdges;
+	TArray<FMeshElement> MarqueeSelectPolygons;
+
+	/** Active transaction while marquee select is in progress
+	    @todo mesheditor: this will be removed when 'current element type' is a thing */
+	TUniquePtr<FScopedTransaction> MarqueeSelectTransaction;
+
+	/** Whether the marquee select transaction is currently active and needs to be ended */
+	bool bMarqueeSelectTransactionActive;
+
+	/** Current view transform.
+	    This is cached from the last known viewport, or taken directly from the VR head transform, if valid */
+	TOptional<FTransform> CachedCameraToWorld;
+
+	/** Requested focus to selection */
+	bool bShouldFocusToSelection;
+
+	/** Whether edits are made per-instance or not */
+	bool bPerInstanceEdits;
+
+	/** Holds all the assets for the mesh editor */
+	class UMeshEditorAssetContainer* AssetContainer;
+
+	friend class FMeshElementViewportTransformable;
+};
+
diff --git a/Engine/Source/Editor/MeshEditor/MeshEditorModeToolkit.cpp b/Engine/Source/Editor/MeshEditor/MeshEditorModeToolkit.cpp
new file mode 100644
--- /dev/null
+++ b/Engine/Source/Editor/MeshEditor/MeshEditorModeToolkit.cpp
@@ -0,0 +1,342 @@
+// Copyright 1998-2016 Epic Games, Inc. All Rights Reserved.
+
+#include "MeshEditorModeToolkit.h"
+#include "MeshEditorMode.h"
+#include "MeshEditorCommands.h"
+#include "MeshEditorStyle.h"
+#include "SWidgetSwitcher.h"
+#include "SSeparator.h"
+#include "SScrollBox.h"
+
+#define LOCTEXT_NAMESPACE "MeshEditorModeToolkit"
+
+
+class SMeshEditorModeControlWidget : public SCompoundWidget
+{
+public:
+	SLATE_BEGIN_ARGS( SMeshEditorModeControlWidget ) {}
+	SLATE_END_ARGS()
+
+public:
+
+	/** SCompoundWidget functions */
+	void Construct( const FArguments& InArgs, const TArray<TTuple<TSharedPtr<FUICommandInfo>, FUIAction>>& Actions )
+	{
+		TSharedRef<SVerticalBox> Buttons = SNew( SVerticalBox );
+		TSharedRef<SVerticalBox> RadioButtons = SNew( SVerticalBox );
+
+		for( const TTuple<TSharedPtr<FUICommandInfo>, FUIAction>& Action : Actions )
+		{
+			const FUICommandInfo& CommandInfo = *Action.Get<0>();
+			const FUIAction& UIAction = Action.Get<1>();
+
+			if( CommandInfo.GetUserInterfaceType() == EUserInterfaceActionType::Button )
+			{
+				Buttons->AddSlot()
+				.AutoHeight()
+				.Padding( 3.0f )
+				[
+					SNew( SButton )
+					.HAlign( HAlign_Center )
+					.VAlign( VAlign_Center )
+					.ContentPadding( FMargin( 4.0f ) )
+					.Text( CommandInfo.GetLabel() )
+					.ToolTip( SNew( SToolTip ).Text( CommandInfo.GetDescription() ) )
+					.OnClicked_Lambda( [UIAction] { UIAction.Execute(); return FReply::Handled(); } )
+					.IsEnabled_Lambda( [UIAction] { return UIAction.CanExecute(); } )
+				];
+			}
+			else if( CommandInfo.GetUserInterfaceType() == EUserInterfaceActionType::RadioButton )
+			{
+				RadioButtons->AddSlot()
+				.AutoHeight()
+				.Padding( 1.0f )
+				[
+					SNew( SCheckBox )
+					.Style( FMeshEditorStyle::Get(), "EditingMode.Entry" )
+					.ToolTip( SNew( SToolTip ).Text( CommandInfo.GetDescription() ) )
+					.IsChecked_Lambda( [UIAction] { return UIAction.GetCheckState(); } )
+					.OnCheckStateChanged_Lambda( [UIAction]( ECheckBoxState State ) { if( State == ECheckBoxState::Checked ) { UIAction.Execute(); } } )
+					[
+						SNew( SOverlay )
+						+SOverlay::Slot()
+						.VAlign( VAlign_Center )
+						[
+							SNew( SSpacer )
+							.Size( FVector2D( 1, 30 ) )
+						]
+						+SOverlay::Slot()
+						.Padding( FMargin( 20, 0, 20, 0 ) )
+						.HAlign( HAlign_Center )
+						.VAlign( VAlign_Center )
+						[
+							SNew( STextBlock )
+							.TextStyle( FMeshEditorStyle::Get(), "EditingMode.Entry.Text" )
+							.Text( CommandInfo.GetLabel() )
+						]
+					]
+				];
+			}
+		}
+
+		// This is the basic layout for each selected element type:
+		// First, a list of buttons, then a separator, and then a list of radio buttons.
+		// @todo mesheditor: if we want to make these UI elements bigger (e.g. for ease of use with VR), we can easily change these
+		// to icons with text to the side. The icon name is already registered with the FUICommandInfo (e.g. "MeshEditorVertex.MoveAction").
+		ChildSlot
+		[
+			SNew( SVerticalBox )
+			+SVerticalBox::Slot()
+			.AutoHeight()
+			.Padding( 6.0f )
+			.HAlign( HAlign_Center )
+			[
+				Buttons
+			]
+			+SVerticalBox::Slot()
+			.AutoHeight()
+			.Padding( 6.0f )
+			[
+				SNew( SSeparator )
+				.Orientation( Orient_Horizontal )
+			]
+			+SVerticalBox::Slot()
+			.AutoHeight()
+			.Padding( 6.0f )
+			.HAlign( HAlign_Center )
+			[
+				RadioButtons
+			]
+		];
+	}
+};
+
+
+class SMeshEditorSelectionModeWidget : public SCompoundWidget
+{
+public:
+	SLATE_BEGIN_ARGS( SMeshEditorModeControlWidget ) {}
+	SLATE_END_ARGS()
+
+public:
+	void Construct( const FArguments& InArgs, FMeshEditorMode& MeshEditorMode, EEditableMeshElementType ElementType, const FText& Label )
+	{
+		ChildSlot
+		[
+			SNew( SBox )
+			.HAlign( HAlign_Fill )
+			.VAlign( VAlign_Center )
+			[
+				SNew( SCheckBox )
+				.Style( FMeshEditorStyle::Get(), "SelectionMode.Entry" )
+				.HAlign( HAlign_Fill )
+				.IsChecked_Lambda( [ &MeshEditorMode, ElementType ] { return ( MeshEditorMode.GetMeshElementSelectionMode() == ElementType ? ECheckBoxState::Checked : ECheckBoxState::Unchecked ); } )
+				.OnCheckStateChanged_Lambda( [ &MeshEditorMode, ElementType ]( ECheckBoxState State ) { if( State == ECheckBoxState::Checked ) { MeshEditorMode.SetMeshElementSelectionMode( ElementType ); } } )
+				[
+					SNew( SHorizontalBox )
+					+SHorizontalBox::Slot()
+					.FillWidth( 1 )
+					.HAlign( HAlign_Center )
+					[
+						SNew( STextBlock )
+						.TextStyle( FMeshEditorStyle::Get(), "SelectionMode.Entry.Text" )
+						.Text( Label )
+					]
+				]
+			]
+		];
+	}
+};
+
+
+void SMeshEditorModeControls::Construct( const FArguments& InArgs )
+{
+	FMeshEditorMode* MeshEditorMode = GetMeshEditorMode();
+	check( MeshEditorMode );
+
+	// Uses the widget switcher widget so only the widget in the slot which corresponds to the selected mesh element type will be shown
+	TSharedRef<SWidgetSwitcher> WidgetSwitcher = SNew( SWidgetSwitcher )
+		.WidgetIndex_Lambda( [MeshEditorMode]() -> int32
+			{
+				if( MeshEditorMode->GetMeshElementSelectionMode() != EEditableMeshElementType::Any )
+				{
+					return static_cast<int32>( MeshEditorMode->GetMeshElementSelectionMode() );
+				}
+				else
+				{
+					return static_cast<int32>( MeshEditorMode->GetSelectedMeshElementType() );
+				}
+			} 
+		);
+
+	WidgetSwitcher->AddSlot( static_cast<int32>( EEditableMeshElementType::Vertex ) )
+	[
+		SNew( SMeshEditorModeControlWidget, MeshEditorMode->GetVertexActions() )
+	];
+
+	WidgetSwitcher->AddSlot( static_cast<int32>( EEditableMeshElementType::Edge ) )
+	[
+		SNew( SMeshEditorModeControlWidget, MeshEditorMode->GetEdgeActions() )
+	];
+
+	WidgetSwitcher->AddSlot( static_cast<int32>( EEditableMeshElementType::Polygon ) )
+	[
+		SNew( SMeshEditorModeControlWidget, MeshEditorMode->GetPolygonActions() )
+	];
+
+	WidgetSwitcher->AddSlot( static_cast<int32>( EEditableMeshElementType::Invalid ) )
+	[
+		SNew( SBox )
+		.Padding( 20.0f )
+		.HAlign( HAlign_Center )
+		[
+			SNew( STextBlock )
+			.Text( LOCTEXT( "NothingSelected", "Nothing is selected" ) )
+		]
+	];
+
+	ChildSlot
+	[
+		SNew( SScrollBox )
+		+SScrollBox::Slot()
+		.Padding( 6.0f )
+		[
+			SNew( SVerticalBox )
+			+SVerticalBox::Slot()
+			.AutoHeight()
+			[
+				SNew( SHorizontalBox )
+				+SHorizontalBox::Slot()
+				.VAlign( VAlign_Center )
+				.AutoWidth()
+				.Padding( 2, 2, 8, 2 )
+				[
+					SNew( STextBlock )
+					.Text( LOCTEXT( "SelectLabel", "Select" ) )
+				]
+				+SHorizontalBox::Slot()
+				.FillWidth( 1 )
+				.Padding( 2 )
+				[
+					SNew( SMeshEditorSelectionModeWidget, *MeshEditorMode, EEditableMeshElementType::Vertex, LOCTEXT( "Vertices", "Vertices" ) )
+				]
+				+SHorizontalBox::Slot()
+				.FillWidth( 1 )
+				.Padding( 2 )
+				[
+					SNew( SMeshEditorSelectionModeWidget, *MeshEditorMode, EEditableMeshElementType::Edge, LOCTEXT( "Edges", "Edges" ) )
+				]
+				+SHorizontalBox::Slot()
+				.FillWidth( 1 )
+				.Padding( 2 )
+				[
+					SNew( SMeshEditorSelectionModeWidget, *MeshEditorMode, EEditableMeshElementType::Polygon, LOCTEXT( "Polygons", "Polygons" ) )
+				]
+				+SHorizontalBox::Slot()
+				.FillWidth( 1 )
+				.Padding( 2 )
+				[
+					SNew( SMeshEditorSelectionModeWidget, *MeshEditorMode, EEditableMeshElementType::Any, LOCTEXT( "Any", "Any" ) )
+				]
+			]
+			+SVerticalBox::Slot()
+			.AutoHeight()
+			.Padding( 0, 4, 0, 0 )
+			.HAlign( HAlign_Right )
+			[
+				SNew( SHorizontalBox )
+				+SHorizontalBox::Slot()
+				.AutoWidth()
+				.Padding( 2 )
+				[
+					SNew( SButton )
+					.HAlign( HAlign_Center )
+					.VAlign( VAlign_Center )
+					.Text( LOCTEXT( "Propagate", "Propagate" ) )
+					.ToolTip( SNew( SToolTip ).Text( LOCTEXT( "PropagateTooltip", "Propagates per-instance changes to the static mesh asset itself." ) ) )
+					.IsEnabled_Lambda( [ MeshEditorMode ] { return false; /*MeshEditorMode->CanPropagateInstanceChanges(); */ } )
+					.OnClicked_Lambda( [ MeshEditorMode ] { MeshEditorMode->PropagateInstanceChanges(); return FReply::Handled(); } )
+				]
+				+SHorizontalBox::Slot()
+				.AutoWidth()
+				.Padding( 2 )
+				[
+					SNew( SBox )
+					.HAlign( HAlign_Fill )
+					.VAlign( VAlign_Center )
+					[
+						SNew( SCheckBox )
+						.Style( FMeshEditorStyle::Get(), "SelectionMode.Entry" )
+						.HAlign( HAlign_Fill )
+						.ToolTip( SNew( SToolTip ).Text( LOCTEXT( "PerInstanceTooltip", "Toggles editing mode between editing instances and editing the original static mesh asset." ) ) )
+						.IsChecked_Lambda( [ MeshEditorMode ] { return ( MeshEditorMode->IsEditingPerInstance() ? ECheckBoxState::Checked : ECheckBoxState::Unchecked ); } )
+						.OnCheckStateChanged_Lambda( [ MeshEditorMode ]( ECheckBoxState State ) { MeshEditorMode->SetEditingPerInstance( State == ECheckBoxState::Checked ); } )
+						[
+							SNew( SHorizontalBox )
+							+SHorizontalBox::Slot()
+							.AutoWidth()
+							.HAlign( HAlign_Center )
+							[
+								SNew( STextBlock )
+								.TextStyle( FMeshEditorStyle::Get(), "SelectionMode.Entry.Text" )
+								.Text( LOCTEXT( "PerInstance", "Per Instance" ) )
+							]
+						]
+					]
+				]
+			]
+			+SVerticalBox::Slot()
+			[
+				SNew( SBorder )
+				.BorderImage( FEditorStyle::GetBrush( "ToolPanel.GroupBorder" ) )
+				[
+					WidgetSwitcher
+				]
+			]
+		]
+	];
+}
+
+
+void FMeshEditorModeToolkit::RegisterTabSpawners( const TSharedRef<FTabManager>& TabManager )
+{
+}
+
+
+void FMeshEditorModeToolkit::UnregisterTabSpawners( const TSharedRef<FTabManager>& TabManager )
+{
+}
+
+
+void FMeshEditorModeToolkit::Init( const TSharedPtr<IToolkitHost>& InitToolkitHost )
+{
+	ToolkitWidget = SNew( SMeshEditorModeControls );
+
+	FModeToolkit::Init( InitToolkitHost );
+}
+
+
+FName FMeshEditorModeToolkit::GetToolkitFName() const
+{
+	return FName( "MeshEditorMode" );
+}
+
+
+FText FMeshEditorModeToolkit::GetBaseToolkitName() const
+{
+	return LOCTEXT( "ToolkitName", "Mesh Editor Mode" );
+}
+
+
+class FEdMode* FMeshEditorModeToolkit::GetEditorMode() const
+{
+	return GLevelEditorModeTools().GetActiveMode( "MeshEditor" );
+}
+
+
+TSharedPtr<SWidget> FMeshEditorModeToolkit::GetInlineContent() const
+{
+	return ToolkitWidget;
+}
+
+#undef LOCTEXT_NAMESPACE
diff --git a/Engine/Source/Editor/MeshEditor/MeshEditorModeToolkit.h b/Engine/Source/Editor/MeshEditor/MeshEditorModeToolkit.h
new file mode 100644
--- /dev/null
+++ b/Engine/Source/Editor/MeshEditor/MeshEditorModeToolkit.h
@@ -0,0 +1,61 @@
+// Copyright 1998-2016 Epic Games, Inc. All Rights Reserved.
+
+#pragma once
+#include "Toolkits/BaseToolkit.h"
+#include "MeshEditorMode.h"
+#include "EditorModeManager.h"
+
+/** Mesh Editor Mode widget for controls */
+class SMeshEditorModeControls : public SCompoundWidget
+{
+public:
+	SLATE_BEGIN_ARGS( SMeshEditorModeControls ) {}
+	SLATE_END_ARGS()
+
+public:
+
+	/** SCompoundWidget functions */
+	void Construct( const FArguments& InArgs );
+
+protected:
+
+	FMeshEditorMode* GetMeshEditorMode() const
+	{
+		return static_cast<FMeshEditorMode*>( GLevelEditorModeTools().GetActiveMode( "MeshEditor" ) );
+	}
+
+private:
+};
+
+
+/**
+ * Mode Toolkit for the Mesh Editor MOde
+ */
+class FMeshEditorModeToolkit : public FModeToolkit
+{
+public:
+
+	FMeshEditorModeToolkit( FMeshEditorMode& InMeshEditorMode )
+		: MeshEditorMode( InMeshEditorMode )
+	{}
+
+	// FModeToolkit overrides
+	virtual void Init( const TSharedPtr<IToolkitHost>& InitToolkitHost ) override;
+
+	// FBaseToolkit overrides
+	virtual void RegisterTabSpawners( const TSharedRef<FTabManager>& TabManager ) override;
+	virtual void UnregisterTabSpawners( const TSharedRef<FTabManager>& TabManager ) override;
+
+	// IToolkit overrides
+	virtual FName GetToolkitFName() const override;
+	virtual FText GetBaseToolkitName() const override;
+	virtual class FEdMode* GetEditorMode() const override;
+	virtual TSharedPtr<SWidget> GetInlineContent() const override;
+
+private:
+	/** Geometry tools widget */
+	TSharedPtr<SMeshEditorModeControls> ToolkitWidget;
+
+	/** Reference to the owning FMeshEditorMode */
+	FMeshEditorMode& MeshEditorMode;
+};
diff --git a/Engine/Source/Editor/MeshEditor/MeshEditorModule.cpp b/Engine/Source/Editor/MeshEditor/MeshEditorModule.cpp
new file mode 100644
--- /dev/null
+++ b/Engine/Source/Editor/MeshEditor/MeshEditorModule.cpp
@@ -0,0 +1,113 @@
+// Copyright 1998-2016 Epic Games, Inc. All Rights Reserved.
+
+#include "MeshEditorModule.h"
+#include "IMeshEditorModule.h"
+#include "EditorModeRegistry.h"
+#include "MeshEditorMode.h"
+#include "MeshEditorStyle.h"
+#include "EditorModeManager.h"
+#include "ISettingsModule.h"
+#include "MeshEditorSettings.h"
+
+#define LOCTEXT_NAMESPACE "MeshEditor"
+
+class FMeshEditorModule : public IMeshEditorModule
+{
+public:
+	FMeshEditorModule()
+	{
+	}
+
+	// FModuleInterface overrides
+	virtual void StartupModule() override;
+	virtual void ShutdownModule() override;
+	virtual void PostLoadCallback() override;
+	virtual bool SupportsDynamicReloading() override
+	{
+		return true;
+	}
+
+	static void ToggleMeshEditorMode();
+
+	virtual IMeshEditorMode* GetLevelEditorMeshEditorMode() override
+	{
+		return static_cast<IMeshEditorMode*>(GLevelEditorModeTools().GetActiveModeTyped<FMeshEditorMode>("MeshEditor"));
+	}
+};
+
+
+namespace MeshEd
+{
+	// @todo mesheditor: get rid of this some day
+	static FAutoConsoleCommand ToggleMeshEditorMode( TEXT( "MeshEd.MeshEditorMode" ), TEXT( "Toggles Mesh Editor Mode" ), FConsoleCommandDelegate::CreateStatic( &FMeshEditorModule::ToggleMeshEditorMode ) );
+}
+
+
+
+void FMeshEditorModule::StartupModule()
+{
+	FMeshEditorStyle::Initialize();
+
+	FEditorModeRegistry::Get().RegisterMode<FMeshEditorMode>(
+		IMeshEditorModule::GetEditorModeID(),
+		LOCTEXT( "ModeName", "Mesh Editor" ),
+		FSlateIcon( FMeshEditorStyle::GetStyleSetName(), "LevelEditor.MeshEditorMode", "LevelEditor.MeshEditorMode.Small" ),
+		true,
+		600
+		);
+//	GLevelEditorModeTools().AddDefaultMode( MeshEditorModeID );
+
+	ISettingsModule* SettingsModule = FModuleManager::GetModulePtr<ISettingsModule>( "Settings" );
+	if( SettingsModule )
+	{
+		// Designer settings
+		SettingsModule->RegisterSettings( "Editor", "ContentEditors", "MeshEditor",
+										  LOCTEXT("MeshEditorSettingsName", "Mesh Editor"),
+										  LOCTEXT("MeshEditorSettingsDescription", "Configure options for the Mesh Editor."),
+										  GetMutableDefault<UMeshEditorSettings>()
+		);
+	}
+}
+
+
+void FMeshEditorModule::ShutdownModule()
+{
+	ISettingsModule* SettingsModule = FModuleManager::GetModulePtr<ISettingsModule>( "Settings" );
+	if( SettingsModule )
+	{
+		SettingsModule->UnregisterSettings( "Editor", "ContentEditors", "MeshEditor" );
+	}
+
+	GLevelEditorModeTools().RemoveDefaultMode(IMeshEditorModule::GetEditorModeID());
+	FEditorModeRegistry::Get().UnregisterMode(IMeshEditorModule::GetEditorModeID());
+
+	FMeshEditorStyle::Shutdown();
+}
+
+
+void FMeshEditorModule::PostLoadCallback()
+{
+	GLevelEditorModeTools().ActivateDefaultMode();
+}
+
+
+void FMeshEditorModule::ToggleMeshEditorMode()
+{
+	if( GLevelEditorModeTools().IsModeActive(IMeshEditorModule::GetEditorModeID()) )
+	{
+		// Shut off Mesh Editor Mode
+		GLevelEditorModeTools().RemoveDefaultMode(IMeshEditorModule::GetEditorModeID());
+		GLevelEditorModeTools().DeactivateMode(IMeshEditorModule::GetEditorModeID());
+	}
+	else
+	{
+		// Activate the mode right away.  We expect it to stay active forever!
+		GLevelEditorModeTools().AddDefaultMode(IMeshEditorModule::GetEditorModeID());
+		GLevelEditorModeTools().ActivateMode(IMeshEditorModule::GetEditorModeID());
+	}
+}
+
+
+IMPLEMENT_MODULE( FMeshEditorModule, MeshEditor )
+
+#undef LOCTEXT_NAMESPACE
diff --git a/Engine/Source/Editor/MeshEditor/MeshEditorModule.h b/Engine/Source/Editor/MeshEditor/MeshEditorModule.h
new file mode 100644
--- /dev/null
+++ b/Engine/Source/Editor/MeshEditor/MeshEditorModule.h
@@ -0,0 +1,6 @@
+// Copyright 1998-2015 Epic Games, Inc. All Rights Reserved.
+
+#pragma once
+
+#include "CoreMinimal.h"
+#include "IMeshEditorModule.h"
diff --git a/Engine/Source/Editor/MeshEditor/MeshEditorSettings.cpp b/Engine/Source/Editor/MeshEditor/MeshEditorSettings.cpp
new file mode 100644
--- /dev/null
+++ b/Engine/Source/Editor/MeshEditor/MeshEditorSettings.cpp
@@ -0,0 +1,15 @@
+// Copyright 1998-2015 Epic Games, Inc. All Rights Reserved.
+
+#include "MeshEditorSettings.h"
+#include "UnrealEdMisc.h"
+
+
+void UMeshEditorSettings::PostEditChangeProperty( struct FPropertyChangedEvent& PropertyChangedEvent )
+{
+	Super::PostEditChangeProperty( PropertyChangedEvent );
+
+	if ( !FUnrealEdMisc::Get().IsDeletePreferences() )
+	{
+		SaveConfig();
+	}
+}
diff --git a/Engine/Source/Editor/MeshEditor/MeshEditorStyle.cpp b/Engine/Source/Editor/MeshEditor/MeshEditorStyle.cpp
new file mode 100644
--- /dev/null
+++ b/Engine/Source/Editor/MeshEditor/MeshEditorStyle.cpp
@@ -0,0 +1,117 @@
+// Copyright 1998-2016 Epic Games, Inc. All Rights Reserved.
+
+#include "MeshEditorStyle.h"
+#include "Styling/SlateStyleRegistry.h"
+
+
+#define IMAGE_BRUSH( RelativePath, ... ) FSlateImageBrush( StyleSet->RootToContentDir( RelativePath, TEXT( ".png" ) ), __VA_ARGS__ )
+#define BOX_BRUSH( RelativePath, ... ) FSlateBoxBrush( StyleSet->RootToContentDir( RelativePath, TEXT( ".png" ) ), __VA_ARGS__ )
+#define TTF_CORE_FONT( RelativePath, ... ) FSlateFontInfo( StyleSet->RootToCoreContentDir( RelativePath, TEXT( ".ttf" ) ), __VA_ARGS__ )
+
+TSharedPtr<FSlateStyleSet> FMeshEditorStyle::StyleSet = nullptr;
+
+void FMeshEditorStyle::Initialize()
+{
+	if( StyleSet.IsValid() )
+	{
+		return;
+	}
+
+	StyleSet = MakeShared<FSlateStyleSet>( GetStyleSetName() );
+
+	// @todo mesheditor: put content in the module itself?
+	StyleSet->SetContentRoot( FPaths::EngineContentDir() / TEXT( "Editor/Slate") );
+	StyleSet->SetCoreContentRoot( FPaths::EngineContentDir() / TEXT( "Slate" ) );
+
+	const FVector2D Icon20x20( 20.0f, 20.0f );
+	const FVector2D Icon40x40( 40.0f, 40.0f );
+	const FVector2D Icon512x512(512.0f, 512.0f);
+	const FLinearColor DimBackground = FLinearColor( FColor( 64, 64, 64 ) );
+	const FLinearColor DimBackgroundHover = FLinearColor( FColor( 80, 80, 80 ) );
+	const FLinearColor LightBackground = FLinearColor( FColor( 128, 128, 128 ) );
+	const FLinearColor HighlightedBackground = FLinearColor( FColor( 255, 192, 0 ) );
+
+	// Icons for the mode panel tabs
+	StyleSet->Set( "LevelEditor.MeshEditorMode", new IMAGE_BRUSH( "Icons/icon_Mode_MeshEdit_40px", Icon40x40 ) );
+	StyleSet->Set( "LevelEditor.MeshEditorMode.Small", new IMAGE_BRUSH( "Icons/icon_Mode_MeshEdit_40px", Icon20x20 ) );
+	StyleSet->Set( "LevelEditor.MeshEditorMode.Selected", new IMAGE_BRUSH( "Icons/icon_Mode_MeshEdit_40px", Icon40x40 ) );
+	StyleSet->Set( "LevelEditor.MeshEditorMode.Selected.Small", new IMAGE_BRUSH( "Icons/icon_Mode_MeshEdit_40px", Icon20x20 ) );
+
+	StyleSet->Set( "EditingMode.Entry", FCheckBoxStyle()
+		.SetCheckBoxType( ESlateCheckBoxType::ToggleButton )
+		.SetUncheckedImage( BOX_BRUSH( "Common/Selection", 8.0f / 32.0f, DimBackground ) )
+		.SetUncheckedPressedImage( BOX_BRUSH( "Common/Selection", 8.0f / 32.0f, LightBackground  ) )
+		.SetUncheckedHoveredImage( BOX_BRUSH( "Common/Selection", 8.0f / 32.0f, DimBackgroundHover ) )
+		.SetCheckedImage( BOX_BRUSH( "Common/Selection", 8.0f / 32.0f, LightBackground ) )
+		.SetCheckedHoveredImage( BOX_BRUSH( "Common/Selection", 8.0f / 32.0f, LightBackground ) )
+		.SetCheckedPressedImage( BOX_BRUSH( "Common/Selection", 8.0f / 32.0f, LightBackground ) )
+		.SetPadding( 0 ) );
+
+	StyleSet->Set( "EditingMode.Entry.Text", FTextBlockStyle()
+		.SetFont( TTF_CORE_FONT( "Fonts/Roboto-Bold", 10 ) )
+		.SetColorAndOpacity( FLinearColor( 1.0f, 1.0f, 1.0f, 0.9f ) )
+		.SetShadowOffset( FVector2D( 1, 1 ) )
+		.SetShadowColorAndOpacity( FLinearColor( 0, 0, 0, 0.9f ) )
+		.SetHighlightColor( FLinearColor( 0.02f, 0.3f, 0.0f ) )
+		.SetHighlightShape( BOX_BRUSH( "Common/TextBlockHighlightShape", FMargin(3.f/8.f) ) ) );
+
+	StyleSet->Set( "SelectionMode.Entry", FCheckBoxStyle()
+		.SetCheckBoxType( ESlateCheckBoxType::ToggleButton )
+		.SetUncheckedImage( BOX_BRUSH( "Common/Button", 8.0f / 32.0f, DimBackground ) )
+		.SetUncheckedPressedImage( BOX_BRUSH( "Common/Button", 8.0f / 32.0f, LightBackground  ) )
+		.SetUncheckedHoveredImage( BOX_BRUSH( "Common/Button", 8.0f / 32.0f, DimBackgroundHover ) )
+		.SetCheckedImage( BOX_BRUSH( "Common/Button", 8.0f / 32.0f, HighlightedBackground ) )
+		.SetCheckedHoveredImage( BOX_BRUSH( "Common/Button", 8.0f / 32.0f, HighlightedBackground ) )
+		.SetCheckedPressedImage( BOX_BRUSH( "Common/Button", 8.0f / 32.0f, HighlightedBackground ) )
+		.SetPadding( 6 ) );
+
+	StyleSet->Set( "SelectionMode.Entry.Text", FTextBlockStyle()
+		.SetFont( TTF_CORE_FONT( "Fonts/Roboto-Regular", 10 ) )
+		.SetColorAndOpacity( FLinearColor( 1.0f, 1.0f, 1.0f, 0.9f ) )
+		.SetHighlightColor( FLinearColor( 0.02f, 0.3f, 0.0f ) )
+		.SetHighlightShape( BOX_BRUSH( "Common/TextBlockHighlightShape", FMargin(3.f/8.f) ) ) );
+	// @todo mesheditor: add icons for other modes / actions
+
+	StyleSet->Set("MeshEditorMode.AddSubdivision", new IMAGE_BRUSH("Icons/VREditor/T_Radial_Sub_Add", Icon512x512));
+	StyleSet->Set("MeshEditorMode.RemoveSubdivision", new IMAGE_BRUSH("Icons/VREditor/T_Radial_Sub_Minus", Icon512x512));
+	StyleSet->Set("MeshEditorMode.PropagateChanges", new IMAGE_BRUSH("Icons/VREditor/Z_Radial_Mesh_Instance", Icon512x512));
+	StyleSet->Set("MeshEditorMode.EditInstance", new IMAGE_BRUSH("Icons/VREditor/Z_Radial_Mesh_Non_Instance", Icon512x512));
+
+	StyleSet->Set("MeshEditorMode.PolyDelete", new IMAGE_BRUSH("Icons/VREditor/T_Radial_Poly_Delete", Icon512x512));
+	StyleSet->Set("MeshEditorMode.PolyExtrude", new IMAGE_BRUSH("Icons/VREditor/T_Radial_Poly_Extrude", Icon512x512));
+	StyleSet->Set("MeshEditorMode.PolyInset", new IMAGE_BRUSH("Icons/VREditor/T_Radial_Poly_Inset", Icon512x512));
+	StyleSet->Set("MeshEditorMode.PolyMove", new IMAGE_BRUSH("Icons/VREditor/T_Radial_Poly_Move", Icon512x512));
+
+	StyleSet->Set("MeshEditorMode.VertexExtend", new IMAGE_BRUSH("Icons/VREditor/T_Radial_Vertex_Extend", Icon512x512));
+	StyleSet->Set("MeshEditorMode.VertexMove", new IMAGE_BRUSH("Icons/VREditor/T_Radial_Vertex_Move", Icon512x512));
+	StyleSet->Set("MeshEditorMode.VertexWeld", new IMAGE_BRUSH("Icons/VREditor/T_Radial_Vertex_Weld", Icon512x512));
+	StyleSet->Set("MeshEditorMode.VertexRemove", new IMAGE_BRUSH("Icons/VREditor/T_Radial_Vertex_Remove", Icon512x512));
+	StyleSet->Set("MeshEditorMode.VertexDelete", new IMAGE_BRUSH("Icons/VREditor/T_Radial_Vertex_Delete", Icon512x512));
+
+	StyleSet->Set("MeshEditorMode.EdgeDelete", new IMAGE_BRUSH("Icons/VREditor/T_Radial_Edge_Delete", Icon512x512));
+	StyleSet->Set("MeshEditorMode.EdgeExtend", new IMAGE_BRUSH("Icons/VREditor/T_Radial_Edge_Extend", Icon512x512));
+	StyleSet->Set("MeshEditorMode.EdgeInsert", new IMAGE_BRUSH("Icons/VREditor/T_Radial_Edge_Insert", Icon512x512));
+	StyleSet->Set("MeshEditorMode.EdgeMove", new IMAGE_BRUSH("Icons/VREditor/T_Radial_Edge_Move", Icon512x512));
+	StyleSet->Set("MeshEditorMode.EdgeRemove", new IMAGE_BRUSH("Icons/VREditor/T_Radial_Edge_Remove", Icon512x512));
+	StyleSet->Set("MeshEditorMode.SelectLoop", new IMAGE_BRUSH("Icons/VREditor/T_Radial_Edge_Select_Loop", Icon512x512));
+
+	FSlateStyleRegistry::RegisterSlateStyle( *StyleSet.Get() );
+}
+
+
+void FMeshEditorStyle::Shutdown()
+{
+	if( StyleSet.IsValid() )
+	{
+		FSlateStyleRegistry::UnRegisterSlateStyle( *StyleSet.Get() );
+		ensure( StyleSet.IsUnique() );
+		StyleSet.Reset();
+	}
+}
+
+
+FName FMeshEditorStyle::GetStyleSetName()
+{
+	static FName StyleName( "MeshEditorStyle" );
+	return StyleName;
+}
diff --git a/Engine/Source/Editor/MeshEditor/MeshEditorStyle.h b/Engine/Source/Editor/MeshEditor/MeshEditorStyle.h
new file mode 100644
--- /dev/null
+++ b/Engine/Source/Editor/MeshEditor/MeshEditorStyle.h
@@ -0,0 +1,24 @@
+// Copyright 1998-2016 Epic Games, Inc. All Rights Reserved.
+
+#pragma once
+
+#include "SlateStyle.h"
+
+class FMeshEditorStyle
+{
+public:
+
+	static void Initialize();
+	static void Shutdown();
+	static TSharedPtr<ISlateStyle> Get() { return StyleSet; }
+	static FName GetStyleSetName();
+
+	static const FSlateBrush* GetBrush(FName PropertyName, const ANSICHAR* Specifier = NULL)
+	{
+		return StyleSet->GetBrush(PropertyName, Specifier);
+	}
+
+private:
+
+	static TSharedPtr<FSlateStyleSet> StyleSet;
+};
diff --git a/Engine/Source/Editor/MeshEditor/MeshElementTransformer.cpp b/Engine/Source/Editor/MeshEditor/MeshElementTransformer.cpp
new file mode 100644
--- /dev/null
+++ b/Engine/Source/Editor/MeshEditor/MeshElementTransformer.cpp
@@ -0,0 +1,20 @@
+// Copyright 1998-2016 Epic Games, Inc. All Rights Reserved.
+
+#include "MeshElementTransformer.h"
+#include "MeshElementViewportTransformable.h"
+#include "ViewportWorldInteraction.h"
+
+
+void UMeshElementTransformer::Init( UViewportWorldInteraction* InitViewportWorldInteraction )
+{
+	Super::Init( InitViewportWorldInteraction );
+
+}
+
+
+void UMeshElementTransformer::Shutdown()
+{
+	Super::Shutdown();
+}
+
+
diff --git a/Engine/Source/Editor/MeshEditor/MeshElementTransformer.h b/Engine/Source/Editor/MeshEditor/MeshElementTransformer.h
new file mode 100644
--- /dev/null
+++ b/Engine/Source/Editor/MeshEditor/MeshElementTransformer.h
@@ -0,0 +1,29 @@
+// Copyright 1998-2016 Epic Games, Inc. All Rights Reserved.
+
+#pragma once
+
+#include "ViewportTransformer.h"
+#include "MeshElementTransformer.generated.h"
+
+
+UCLASS()
+class UMeshElementTransformer : public UViewportTransformer
+{
+	GENERATED_BODY()
+
+public:
+
+	// UViewportTransformer overrides
+	virtual void Init( class UViewportWorldInteraction* InitViewportWorldInteraction ) override;
+	virtual void Shutdown() override;
+	virtual bool ShouldCenterTransformGizmoPivot() const override
+	{
+		// Mesh elements always default to being pivoted from the center of all elements' bounds
+		return true;
+	}
+
+protected:
+
+
+};
+
diff --git a/Engine/Source/Editor/MeshEditor/MeshElementViewportTransformable.cpp b/Engine/Source/Editor/MeshEditor/MeshElementViewportTransformable.cpp
new file mode 100644
--- /dev/null
+++ b/Engine/Source/Editor/MeshEditor/MeshElementViewportTransformable.cpp
@@ -0,0 +1,83 @@
+// Copyright 1998-2016 Epic Games, Inc. All Rights Reserved.
+
+#include "MeshElementViewportTransformable.h"
+#include "MeshEditorMode.h"
+#include "EditableMesh.h"
+
+
+const FTransform FMeshElementViewportTransformable::GetTransform() const
+{
+	return CurrentTransform;
+}
+
+
+bool FMeshElementViewportTransformable::IsUnorientedPoint() const
+{
+	// Vertex elements never support rotation or scaling when only one is selected
+	return MeshElement.ElementAddress.ElementType == EEditableMeshElementType::Vertex;
+}
+
+
+void FMeshElementViewportTransformable::ApplyTransform( const FTransform& NewTransform, const bool bSweep )
+{
+	CurrentTransform = NewTransform;
+}
+
+
+FBox FMeshElementViewportTransformable::BuildBoundingBox( const FTransform& BoundingBoxToWorld ) const
+{
+	FBox BoundingBox = FBox( 0 );
+
+	if( MeshElement.IsValidMeshElement() )
+	{
+		const FTransform WorldToBoundingBox = BoundingBoxToWorld.Inverse();
+		const FTransform& ComponentToWorld = MeshElement.Component.Get()->GetComponentToWorld();
+		const FTransform ComponentToBoundingBox = ComponentToWorld * WorldToBoundingBox;
+
+		UEditableMesh* EditableMesh = MeshEditorMode.FindOrCreateEditableMesh( *MeshElement.Component, MeshElement.ElementAddress.SubMeshAddress );
+		if( EditableMesh != nullptr )
+		{
+			if( FMeshEditorMode::IsElementIDValid( MeshElement, EditableMesh ) )
+			{
+				BoundingBox.Init();
+
+				FTransform ElementTransform = FTransform::Identity;
+				switch( MeshElement.ElementAddress.ElementType )
+				{
+					case EEditableMeshElementType::Vertex:
+						BoundingBox += ComponentToBoundingBox.TransformPosition( EditableMesh->GetVertexAttribute( FVertexID( MeshElement.ElementAddress.ElementID ), UEditableMeshAttribute::VertexPosition(), 0 ) );
+						break;
+
+					case EEditableMeshElementType::Edge:
+					{
+						FVertexID EdgeVertexID0, EdgeVertexID1;
+						EditableMesh->GetEdgeVertices( FEdgeID( MeshElement.ElementAddress.ElementID ), /* Out */ EdgeVertexID0, /* Out */ EdgeVertexID1 );
+
+						BoundingBox += ComponentToBoundingBox.TransformPosition( EditableMesh->GetVertexAttribute( EdgeVertexID0, UEditableMeshAttribute::VertexPosition(), 0 ) );
+						BoundingBox += ComponentToBoundingBox.TransformPosition( EditableMesh->GetVertexAttribute( EdgeVertexID1, UEditableMeshAttribute::VertexPosition(), 0 ) );
+					}
+					break;
+
+					case EEditableMeshElementType::Polygon:
+					{
+						const FPolygonRef PolygonRef( MeshElement.ElementAddress.SectionID, FPolygonID( MeshElement.ElementAddress.ElementID ) );
+
+						static TArray<FVertexID> PerimeterVertexIDs;
+						EditableMesh->GetPolygonPerimeterVertices( PolygonRef, /* Out */ PerimeterVertexIDs );
+
+						for( const FVertexID VertexID : PerimeterVertexIDs )
+						{
+							BoundingBox += ComponentToBoundingBox.TransformPosition( EditableMesh->GetVertexAttribute( VertexID, UEditableMeshAttribute::VertexPosition(), 0 ) );
+						}
+					}
+					break;
+
+					default:
+						check( 0 );
+				}
+			}
+		}
+	}
+
+	return BoundingBox;
+}
\ No newline at end of file
diff --git a/Engine/Source/Editor/MeshEditor/MeshElementViewportTransformable.h b/Engine/Source/Editor/MeshEditor/MeshElementViewportTransformable.h
new file mode 100644
--- /dev/null
+++ b/Engine/Source/Editor/MeshEditor/MeshElementViewportTransformable.h
@@ -0,0 +1,37 @@
+// Copyright 1998-2016 Epic Games, Inc. All Rights Reserved.
+
+#pragma once
+
+#include "ViewportTransformable.h"
+#include "MeshElement.h"
+
+/**
+ * A transformable mesh element
+ */
+class FMeshElementViewportTransformable : public FViewportTransformable
+{
+public:
+
+	/** Sets up safe defaults */
+	FMeshElementViewportTransformable( class FMeshEditorMode& InitMeshEditorMode )
+		: MeshEditorMode( InitMeshEditorMode )
+	{
+	}
+
+	// FViewportTransformable overrides
+	virtual const FTransform GetTransform() const override;
+	virtual bool IsUnorientedPoint() const override;
+	virtual void ApplyTransform( const FTransform& NewTransform, const bool bSweep ) override;
+	virtual FBox BuildBoundingBox( const FTransform& BoundingBoxToWorld ) const override;
+
+
+	/** Mesh editor mode object */
+	FMeshEditorMode& MeshEditorMode;
+
+	/** The actual mesh element being transformed */
+	FMeshElement MeshElement;
+
+	/** Current transform for the element */
+	FTransform CurrentTransform;
+};
+
diff --git a/Engine/Source/Editor/MeshEditor/Public/IMeshEditorMode.h b/Engine/Source/Editor/MeshEditor/Public/IMeshEditorMode.h
new file mode 100644
--- /dev/null
+++ b/Engine/Source/Editor/MeshEditor/Public/IMeshEditorMode.h
@@ -0,0 +1,23 @@
+// Copyright 1998-2017 Epic Games, Inc. All Rights Reserved.
+#pragma once
+
+#include "CoreMinimal.h"
+#include "EdMode.h"
+
+enum class EEditableMeshElementType;
+
+class IMeshEditorMode 
+{
+public:
+	/** Returns the current element selection mode in operation */
+	virtual EEditableMeshElementType GetMeshElementSelectionMode() const = 0;
+
+	/** Sets the element selection mode to use */
+	virtual void SetMeshElementSelectionMode(EEditableMeshElementType ElementType) = 0;
+
+	/** Set up hooks for external menu systems like VR Editor */
+	virtual void RegisterWithExternalMenuSystem() = 0;
+
+	/** Unregister from any external menu systems */
+	virtual void UnregisterWithExternalMenuSystem() = 0;
+};
\ No newline at end of file
diff --git a/Engine/Source/Editor/MeshEditor/Public/IMeshEditorModule.h b/Engine/Source/Editor/MeshEditor/Public/IMeshEditorModule.h
new file mode 100644
--- /dev/null
+++ b/Engine/Source/Editor/MeshEditor/Public/IMeshEditorModule.h
@@ -0,0 +1,51 @@
+// Copyright 1998-2015 Epic Games, Inc. All Rights Reserved.
+
+#pragma once
+
+#include "ModuleManager.h"
+#include "Editor.h"
+#include "IMeshEditorMode.h"
+#include "EditorModeManager.h"
+
+
+/**
+ * The public interface to this module
+ */
+class IMeshEditorModule : public IModuleInterface
+{
+
+public:
+
+	/**
+	 * Singleton-like access to this module's interface.  This is just for convenience!
+	 * Beware of calling this during the shutdown phase, though.  Your module might have been unloaded already.
+	 *
+	 * @return Returns singleton instance, loading the module on demand if needed
+	 */
+	static inline IMeshEditorModule& Get()
+	{
+		return FModuleManager::LoadModuleChecked< IMeshEditorModule >( "MeshEditor" );
+	}
+
+	/**
+	 * Checks to see if this module is loaded and ready.  It is only valid to call Get() if IsAvailable() returns true.
+	 *
+	 * @return True if the module is loaded and ready to use
+	 */
+	static inline bool IsAvailable()
+	{
+		return FModuleManager::Get().IsModuleLoaded( "MeshEditor" );
+	}
+
+	/** Stores the FEditorMode ID of the associated editor mode */
+	static inline FEditorModeID GetEditorModeID()
+	{
+		static FEditorModeID MeshEditorFeatureName = FName(TEXT("MeshEditor"));
+		return MeshEditorFeatureName;
+	}
+
+	/** Returns the Level Editor's Mesh Editor Mode as an interface for external systems */
+	virtual IMeshEditorMode* GetLevelEditorMeshEditorMode() = 0;
+
+};
+
diff --git a/Engine/Source/Editor/MeshEditor/Public/MeshEditorAssetContainer.h b/Engine/Source/Editor/MeshEditor/Public/MeshEditorAssetContainer.h
new file mode 100644
--- /dev/null
+++ b/Engine/Source/Editor/MeshEditor/Public/MeshEditorAssetContainer.h
@@ -0,0 +1,24 @@
+// Copyright 1998-2017 Epic Games, Inc. All Rights Reserved.
+#pragma once
+
+#include "Engine/DataAsset.h"
+#include "MeshEditorAssetContainer.generated.h"
+
+/**
+ * Asset container for the mesh editor
+ */
+UCLASS()
+class MESHEDITOR_API UMeshEditorAssetContainer : public UDataAsset
+{
+    GENERATED_BODY()
+ 
+public:
+	UPROPERTY(EditAnywhere, Category = Material)
+	class UMaterial* HoveredGeometryMaterial;
+
+	UPROPERTY(EditAnywhere, Category = Material)
+	class UMaterial* HoveredFaceMaterial;
+
+	UPROPERTY(EditAnywhere, Category = Sound)
+	class USoundBase* DefaultSound;
+};
\ No newline at end of file
diff --git a/Engine/Source/Editor/MeshEditor/Public/MeshEditorSettings.h b/Engine/Source/Editor/MeshEditor/Public/MeshEditorSettings.h
new file mode 100644
--- /dev/null
+++ b/Engine/Source/Editor/MeshEditor/Public/MeshEditorSettings.h
@@ -0,0 +1,59 @@
+// Copyright 1998-2017 Epic Games, Inc. All Rights Reserved.
+
+#pragma once
+
+
+#include "MeshEditorSettings.generated.h"
+
+/**
+ * Implements the Mesh Editor's settings.
+ */
+UCLASS(config=EditorPerProjectUserSettings)
+class MESHEDITOR_API UMeshEditorSettings
+	: public UObject
+{
+	GENERATED_BODY()
+
+public:
+
+	UMeshEditorSettings()
+		: bSeparateSelectionSetPerMode( false )
+		, bOnlyEditSelectedObjects( true )
+		, bOnlySelectVisibleMeshes( true )
+		, bOnlySelectVisibleElements( true )
+		, bAllowGrabberSphere( false )
+		, bAutoQuadrangulate( false )
+	{}
+
+public:
+
+	/** If set, each element selection mode remembers its own selection set. Otherwise, changing selection mode adapts the current selection as appropriate. */
+	UPROPERTY(EditAnywhere, config, Category=MeshEditor)
+	bool bSeparateSelectionSetPerMode;
+
+	/** Whether you must select objects before mesh editing features will be available.  This helps to prevent accidental mis-clicks. */
+	UPROPERTY(EditAnywhere, config, Category=MeshEditor)
+	bool bOnlyEditSelectedObjects;
+
+	/** Whether only unoccluded meshes will be selected by marquee select, or whether all meshes within the selection box will be selected, regardless of whether they are behind another. */
+	UPROPERTY(EditAnywhere, config, Category=MeshEditor)
+	bool bOnlySelectVisibleMeshes;
+
+	/** Whether only front-facing vertices, edges or polygons will be selected by marquee select. */
+	UPROPERTY(EditAnywhere, config, Category=MeshEditor)
+	bool bOnlySelectVisibleElements;
+
+	/** When enabled, the grabber sphere will be used to select and move mesh elements near the interactor's origin */
+	UPROPERTY(EditAnywhere, config, Category=MeshEditor)
+	bool bAllowGrabberSphere;
+
+	/** When enabled, triangulated static meshes will be auto-quadrangulated when converted to editable meshes */
+	UPROPERTY(EditAnywhere, config, Category=MeshEditor)
+	bool bAutoQuadrangulate;
+
+protected:
+
+	// UObject overrides
+	virtual void PostEditChangeProperty( struct FPropertyChangedEvent& PropertyChangedEvent ) override;
+
+};
diff --git a/Engine/Source/Editor/MeshPaint/Private/MeshPaintEdMode.cpp b/Engine/Source/Editor/MeshPaint/Private/MeshPaintEdMode.cpp
--- a/Engine/Source/Editor/MeshPaint/Private/MeshPaintEdMode.cpp
+++ b/Engine/Source/Editor/MeshPaint/Private/MeshPaintEdMode.cpp
@@ -27,7 +27,6 @@
 
 #include "StaticMeshResources.h"
 #include "VREditorMode.h"
-#include "IVREditorModule.h"
 #include "ScopedTransaction.h"
 #include "MeshPaintRendering.h"
 #include "Toolkits/ToolkitManager.h"
@@ -50,7 +49,6 @@
 #include "ViewportWorldInteraction.h"
 #include "ViewportInteractableInterface.h"
 #include "VREditorInteractor.h"
-#include "EditorWorldManager.h"
 #include "UniquePtr.h"
 
 #define LOCTEXT_NAMESPACE "MeshPaint_Mode"
@@ -257,18 +255,18 @@ void FEdModeMeshPaint::Enter()
 		ApplyOrRemoveForceBestLOD(/*bApply=*/ true, PaintingMeshLODIndex);
 	}
 
-	TSharedPtr<FEditorWorldWrapper> EditorWorld = GEditor->GetEditorWorldManager()->GetEditorWorldWrapper( GetWorld() );
-	if(EditorWorld.IsValid())
+	UEditorWorldExtensionCollection* ExtensionCollection = GEditor->GetEditorWorldExtensionsManager()->GetEditorWorldExtensions(GetWorld());
+	if(ExtensionCollection != nullptr)
 	{
 		// Register to find out about VR input events
-		UViewportWorldInteraction* WorldInteraction = EditorWorld->GetViewportWorldInteraction();
-		if(WorldInteraction != nullptr)
+		UViewportWorldInteraction* ViewportWorldInteraction = Cast<UViewportWorldInteraction>(ExtensionCollection->FindExtension(UViewportWorldInteraction::StaticClass()));
+		if(ViewportWorldInteraction != nullptr)
 		{
-			WorldInteraction->OnViewportInteractionInputAction().RemoveAll(this);
-			WorldInteraction->OnViewportInteractionInputAction().AddRaw(this, &FEdModeMeshPaint::OnVRAction);
+			ViewportWorldInteraction->OnViewportInteractionInputAction().RemoveAll(this);
+			ViewportWorldInteraction->OnViewportInteractionInputAction().AddRaw(this, &FEdModeMeshPaint::OnVRAction);
 
 			// Hide the VR transform gizmo while we're in mesh paint mode.  It sort of gets in the way of painting.
-			WorldInteraction->SetTransformGizmoVisible(false);
+			ViewportWorldInteraction->SetTransformGizmoVisible(false);
 		}
 	}
 }
@@ -276,16 +274,19 @@ void FEdModeMeshPaint::Enter()
 /** FEdMode: Called when the mode is exited */
 void FEdModeMeshPaint::Exit()
 {
-	if( IVREditorModule::IsAvailable() )
 	{
-		UViewportWorldInteraction* ViewpportWorldInteraction = GEditor->GetEditorWorldManager()->GetEditorWorldWrapper(GetWorld())->GetViewportWorldInteraction();
-		if(ViewpportWorldInteraction != nullptr )
+		UEditorWorldExtensionCollection* ExtensionCollection = GEditor->GetEditorWorldExtensionsManager()->GetEditorWorldExtensions(GetWorld());
+		if (ExtensionCollection != nullptr)
 		{
-			// Restore the transform gizmo visibility
-			ViewpportWorldInteraction->SetTransformGizmoVisible( true );
+			UViewportWorldInteraction* ViewportWorldInteraction = Cast<UViewportWorldInteraction>(ExtensionCollection->FindExtension(UViewportWorldInteraction::StaticClass()));
+			if(ViewportWorldInteraction != nullptr )
+			{
+				// Restore the transform gizmo visibility
+				ViewportWorldInteraction->SetTransformGizmoVisible( true );
 
-			// Unregister from event handlers
-			ViewpportWorldInteraction->OnViewportInteractionInputAction().RemoveAll( this );
+				// Unregister from event handlers
+				ViewportWorldInteraction->OnViewportInteractionInputAction().RemoveAll( this );
+			}
 		}
 	}
 
@@ -2749,11 +2750,9 @@ void FEdModeMeshPaint::Render( const FSceneView* View, FViewport* Viewport, FPri
 		static TArray<FPaintRay> PaintRays;
 		PaintRays.Reset();
 		
-		UVREditorMode* VREditorMode = GEditor->GetEditorWorldManager()->GetEditorWorldWrapper( Viewport->GetClient()->GetWorld() )->GetVREditorMode();
-
+		UVREditorMode* VREditorMode = Cast<UVREditorMode>( GEditor->GetEditorWorldExtensionsManager()->GetEditorWorldExtensions( GetWorld() )->FindExtension( UVREditorMode::StaticClass() ) );
 		// Check to see if VREditorMode is active. If so, we're painting with interactor
 		bool bIsInVRMode = false;
-		if ( IVREditorModule::IsAvailable() )
 		{
 			if ( VREditorMode != nullptr && VREditorMode->IsFullyInitialized() && VREditorMode->IsActive())
 			{
@@ -3676,7 +3675,7 @@ EMeshPaintAction::Type FEdModeMeshPaint::GetPaintAction(FViewport* InViewport)
 		// When painting using VR, allow the modifier button to activate Erase mode
 		if( PaintingWithInteractorInVR )
 		{
-			UVREditorMode* VREditorMode = GEditor->GetEditorWorldManager()->GetEditorWorldWrapper( GetWorld() )->GetVREditorMode();
+			UVREditorMode* VREditorMode = Cast<UVREditorMode>( GEditor->GetEditorWorldExtensionsManager()->GetEditorWorldExtensions( GetWorld() )->FindExtension( UVREditorMode::StaticClass() ) );
 			UVREditorInteractor* VRInteractor = Cast<UVREditorInteractor>( PaintingWithInteractorInVR );
 			if( VREditorMode != nullptr && VRInteractor != nullptr )
 			{
@@ -5428,11 +5427,11 @@ TWeakObjectPtr<AActor> FEdModeMeshPaint::GetEditingActor() const
 void FEdModeMeshPaint::OnVRAction( class FEditorViewportClient& ViewportClient, UViewportInteractor* Interactor, 
 	const FViewportActionKeyInput& Action, bool& bOutIsInputCaptured, bool& bWasHandled )
 {
-	UVREditorMode* VREditorMode = GEditor->GetEditorWorldManager()->GetEditorWorldWrapper(ViewportClient.GetWorld())->GetVREditorMode();
+	UVREditorMode* VREditorMode = Cast<UVREditorMode>( GEditor->GetEditorWorldExtensionsManager()->GetEditorWorldExtensions( GetWorld() )->FindExtension( UVREditorMode::StaticClass() ) );
 	UVREditorInteractor* VRInteractor = Cast<UVREditorInteractor>( Interactor );
-	if( VREditorMode != nullptr && VREditorMode->IsActive() && VRInteractor != nullptr )
+	if( VREditorMode != nullptr && VREditorMode->IsActive() && VRInteractor != nullptr && VRInteractor->GetDraggingMode() == EViewportInteractionDraggingMode::Nothing)
 	{
-		if( Action.ActionType == ViewportWorldActionTypes::SelectAndMove_LightlyPressed )
+		if( Action.ActionType == ViewportWorldActionTypes::SelectAndMove )
 		{
 			if( !bIsPainting && Action.Event == IE_Pressed && !VRInteractor->IsHoveringOverPriorityType() )
 			{
@@ -5470,10 +5469,6 @@ void FEdModeMeshPaint::OnVRAction( class FEditorViewportClient& ViewportClient,
 					bWasHandled = true;
 					bOutIsInputCaptured = true;
 
-					// Don't allow a "full press" to happen at all with this trigger pull.  We don't need full press for anything, and it
-					// would interrupt our light press.
-					VRInteractor->SetAllowTriggerFullPress( false );
-
 					StartPainting();
 					PaintingWithInteractorInVR = Interactor;
 
@@ -5503,6 +5498,12 @@ void FEdModeMeshPaint::OnVRAction( class FEditorViewportClient& ViewportClient,
 				bWasHandled = true;
 				bOutIsInputCaptured = false;
 			}
+
+			else if( bIsPainting )
+			{
+				// A different hand might be painting, so absorb the input
+				bOutIsInputCaptured = bWasHandled = ( Action.Event == IE_Released ? false : true );
+			}
 		}
 	}
 }
diff --git a/Engine/Source/Editor/MovieSceneTools/Private/Sections/ThumbnailSection.cpp b/Engine/Source/Editor/MovieSceneTools/Private/Sections/ThumbnailSection.cpp
--- a/Engine/Source/Editor/MovieSceneTools/Private/Sections/ThumbnailSection.cpp
+++ b/Engine/Source/Editor/MovieSceneTools/Private/Sections/ThumbnailSection.cpp
@@ -15,6 +15,7 @@
 #include "Widgets/Text/SInlineEditableTextBlock.h"
 #include "PropertyEditorModule.h"
 #include "IDetailsView.h"
+#include "IVREditorModule.h"
 
 
 #define LOCTEXT_NAMESPACE "FThumbnailSection"
@@ -228,7 +229,7 @@ int32 FThumbnailSection::OnPaintSection( FSequencerSectionPainter& InPainter ) c
 
 	static const float SectionThumbnailPadding = 4.f;
 
-	const ESlateDrawEffect::Type DrawEffects = InPainter.bParentEnabled ? ESlateDrawEffect::None : ESlateDrawEffect::DisabledEffect;
+	ESlateDrawEffect::Type DrawEffects = InPainter.bParentEnabled ? ESlateDrawEffect::None : ESlateDrawEffect::DisabledEffect;
 
 	int32 LayerId = InPainter.LayerId;
 
@@ -267,13 +268,22 @@ int32 FThumbnailSection::OnPaintSection( FSequencerSectionPainter& InPainter ) c
 
 		if (Fade <= 1.f)
 		{
+			DrawEffects |= ESlateDrawEffect::NoGamma;
+
+			if (IVREditorModule::Get().IsVREditorModeActive())
+			{
+				// In VR editor every widget is in the world and gamma corrected by the scene renderer.  Thumbnails will have already been gamma
+				// corrected and so they need to be reversed
+				DrawEffects |= ESlateDrawEffect::ReverseGamma;
+			}
+
 			FSlateDrawElement::MakeViewport(
 				InPainter.DrawElements,
 				LayerId,
 				PaintGeometry,
 				Thumbnail,
 				ThumbnailClipRect,
-				DrawEffects | ESlateDrawEffect::NoGamma,
+				DrawEffects,
 				FLinearColor(1.f, 1.f, 1.f, 1.f - Fade)
 				);
 		}
diff --git a/Engine/Source/Editor/MovieSceneTools/Private/TrackEditors/TransformTrackEditor.cpp b/Engine/Source/Editor/MovieSceneTools/Private/TrackEditors/TransformTrackEditor.cpp
--- a/Engine/Source/Editor/MovieSceneTools/Private/TrackEditors/TransformTrackEditor.cpp
+++ b/Engine/Source/Editor/MovieSceneTools/Private/TrackEditors/TransformTrackEditor.cpp
@@ -437,20 +437,32 @@ TSharedRef<ISequencerSection> F3DTransformTrackEditor::MakeSectionInterface(UMov
 void F3DTransformTrackEditor::OnPreTransformChanged( UObject& InObject )
 {
 	float AutoKeyTime = GetTimeForKey();
-
-	if( IsAllowedToAutoKey() )
+	AActor* Actor = Cast<AActor>(&InObject);
+	// If Sequencer is allowed to autokey and we are clicking on an Actor that can be autokeyed
+	if( IsAllowedToAutoKey() && (Actor && !Actor->IsEditorOnly() ) )
 	{
-		AActor* Actor = nullptr;
+		AActor* ActorThatChanged = nullptr;
 		USceneComponent* SceneComponentThatChanged = nullptr;
-		GetActorAndSceneComponentFromObject(&InObject, Actor, SceneComponentThatChanged);
+		GetActorAndSceneComponentFromObject(&InObject, ActorThatChanged, SceneComponentThatChanged);
+		USceneComponent* SceneComponent = nullptr;
+		
+		if( ActorThatChanged && ActorThatChanged->GetRootComponent() )
+		{
+			SceneComponent = ActorThatChanged->GetRootComponent();
+		}
+		else
+		{
+			// If the object wasn't an actor attempt to get it directly as a scene component 
+			SceneComponent = Cast<USceneComponent>( &InObject );
+		}
 
 		if( SceneComponentThatChanged )
 		{
 			// Cache off the existing transform so we can detect which components have changed
 			// and keys only when something has changed
 			FTransformData Transform( SceneComponentThatChanged );
-
-			ObjectToExistingTransform.Add( &InObject, Transform );
+			
+			ObjectToExistingTransform.Add(&InObject, Transform);
 		}
 	}
 }
@@ -477,7 +489,8 @@ void F3DTransformTrackEditor::OnTransformChanged( UObject& InObject )
 	USceneComponent* SceneComponentThatChanged = nullptr;
 	GetActorAndSceneComponentFromObject(&InObject, Actor, SceneComponentThatChanged);
 
-	if( SceneComponentThatChanged != nullptr && Actor != nullptr )
+	// If the Actor that just finished transforming doesn't have autokey disabled
+	if( SceneComponentThatChanged != nullptr && (Actor && !Actor->IsEditorOnly()))
 	{
 		// Find an existing transform if possible.  If one exists we will compare against the new one to decide what components of the transform need keys
 		FTransformData ExistingTransform = ObjectToExistingTransform.FindRef( &InObject );
diff --git a/Engine/Source/Editor/Sequencer/Private/Sequencer.h b/Engine/Source/Editor/Sequencer/Private/Sequencer.h
--- a/Engine/Source/Editor/Sequencer/Private/Sequencer.h
+++ b/Engine/Source/Editor/Sequencer/Private/Sequencer.h
@@ -371,7 +371,6 @@ public:
 	void RecordSelectedActors();
 
 	/** Functions to push on to the transport controls we use */
-	FReply OnPlay(bool bTogglePlay=true, float InPlayRate=1.f);
 	FReply OnRecord();
 	FReply OnStepForward();
 	FReply OnStepBackward();
@@ -392,6 +391,7 @@ public:
 	/** Delegate handler for recording finishing */
 	void HandleRecordingFinished(UMovieSceneSequence* Sequence);
 
+	/** Whether or not playback is looping */
 	bool IsLooping() const;
 
 	/** Set the new global time, accounting for looping options */
@@ -612,11 +612,14 @@ public:
 	virtual TSharedRef<INumericTypeInterface<float>> GetNumericTypeInterface() override;
 	virtual TSharedRef<INumericTypeInterface<float>> GetZeroPadNumericTypeInterface() override;
 	virtual TSharedRef<SWidget> MakeTransportControls(bool bExtended) override;
+	virtual FReply OnPlay(bool bTogglePlay = true, float InPlayRate = 1.f) override;
+	virtual void Pause() override;
 	virtual TSharedRef<SWidget> MakeTimeRange(const TSharedRef<SWidget>& InnerContent, bool bShowWorkingRange, bool bShowViewRange, bool bShowPlaybackRange) override;
 	virtual UObject* FindSpawnedObjectOrTemplate(const FGuid& BindingId) override;
 	virtual FGuid MakeNewSpawnable(UObject& SourceObject) override;
 	virtual bool IsReadOnly() const override;
 	virtual void ExternalSelectionHasChanged() override { SynchronizeSequencerSelectionWithExternalSelection(); }
+	virtual USequencerSettings* GetSequencerSettings() override { return Settings; }
 
 public:
 
@@ -746,7 +749,6 @@ protected:
 	void Rewind();
 	void ShuttleForward();
 	void ShuttleBackward();
-	void Pause();
 	void StepForward();
 	void StepBackward();
 	void StepToNextKey();
diff --git a/Engine/Source/Editor/Sequencer/Private/SequencerKeyActor.cpp b/Engine/Source/Editor/Sequencer/Private/SequencerKeyActor.cpp
new file mode 100644
--- /dev/null
+++ b/Engine/Source/Editor/Sequencer/Private/SequencerKeyActor.cpp
@@ -0,0 +1,92 @@
+// Copyright 1998-2016 Epic Games, Inc. All Rights Reserved.
+
+#include "SequencerKeyActor.h"
+#include "Sections/MovieScene3DTransformSection.h"
+#include "SequencerEdMode.h"
+#include "Materials/MaterialInstance.h" 
+#include "Materials/MaterialInstanceDynamic.h"
+#include "Engine/Selection.h"
+#include "EditorModeTools.h"
+#include "EditorModeManager.h"
+#include "ModuleManager.h"
+
+ASequencerKeyActor::ASequencerKeyActor()
+	: Super()
+{
+	UStaticMesh* KeyEditorMesh = LoadObject<UStaticMesh>(nullptr, TEXT("/Engine/VREditor/TransformGizmo/SM_Sequencer_Key"));
+	check(KeyEditorMesh != nullptr);
+	UMaterial* KeyEditorMaterial = LoadObject<UMaterial>(nullptr, TEXT("/Engine/VREditor/TransformGizmo/Main"));
+	check(KeyEditorMaterial != nullptr);
+
+	const bool bTransient = true;
+	USceneComponent* SceneComponent = CreateDefaultSubobject<USceneComponent>(TEXT("SceneComponent"), bTransient);
+	check(SceneComponent != nullptr);
+	this->RootComponent = SceneComponent;
+
+	KeyMeshComponent = CreateDefaultSubobject<UStaticMeshComponent>(TEXT("KeyMesh"));
+	KeyMeshComponent->SetMobility(EComponentMobility::Movable);
+	KeyMeshComponent->SetupAttachment(RootComponent);
+	KeyMeshComponent->SetStaticMesh(KeyEditorMesh);
+	KeyMeshComponent->CreateAndSetMaterialInstanceDynamicFromMaterial(0, KeyEditorMaterial);
+
+	KeyMeshComponent->SetCollisionEnabled(ECollisionEnabled::QueryOnly);
+	KeyMeshComponent->SetCollisionResponseToAllChannels(ECR_Ignore);
+	KeyMeshComponent->SetCollisionResponseToChannel(ECC_Visibility, ECR_Block);
+	
+	KeyMeshComponent->bGenerateOverlapEvents = false;
+	KeyMeshComponent->SetCanEverAffectNavigation(false);
+	KeyMeshComponent->bCastDynamicShadow = false;
+	KeyMeshComponent->bCastStaticShadow = false;
+	KeyMeshComponent->bAffectDistanceFieldLighting = false;
+	KeyMeshComponent->bAffectDynamicIndirectLighting = false;
+}
+
+
+void ASequencerKeyActor::PostEditMove(bool bFinished)
+{
+	// Push the key's transform to the Sequencer track
+	PropagateKeyChange();
+	Super::PostEditMove(bFinished);
+}
+
+void ASequencerKeyActor::SetKeyData(class UMovieScene3DTransformSection* NewTrackSection, float NewKeyTime)
+{
+	TrackSection = NewTrackSection;
+	KeyTime = NewKeyTime;
+	// Associate the currently selected actor with this key
+	AssociatedActor = GEditor->GetSelectedActors()->GetTop<AActor>();
+	// Draw a single transform track based on the data from this key
+	FEditorViewportClient* ViewportClient = StaticCast<FEditorViewportClient*>(GEditor->GetActiveViewport()->GetClient());
+	if (ViewportClient != nullptr)
+	{
+		FSequencerEdMode* SequencerEdMode = (FSequencerEdMode*)(ViewportClient->GetModeTools()->GetActiveMode(FSequencerEdMode::EM_SequencerMode));
+		SequencerEdMode->DrawMeshTransformTrailFromKey(this);
+	}
+}
+
+
+void ASequencerKeyActor::PropagateKeyChange()
+{
+	if (TrackSection != nullptr)
+	{
+		// Mark the track section as dirty
+		TrackSection->Modify();
+
+		// Update the translation keys
+		FRichCurve& TransXCurve = TrackSection->GetTranslationCurve(EAxis::X);
+		FRichCurve& TransYCurve = TrackSection->GetTranslationCurve(EAxis::Y);
+		FRichCurve& TransZCurve = TrackSection->GetTranslationCurve(EAxis::Z);
+		TransXCurve.UpdateOrAddKey(KeyTime, GetActorTransform().GetLocation().X);
+		TransYCurve.UpdateOrAddKey(KeyTime, GetActorTransform().GetLocation().Y);
+		TransZCurve.UpdateOrAddKey(KeyTime, GetActorTransform().GetLocation().Z);
+
+		// Draw a single transform track based on the data from this key
+		FEditorViewportClient* ViewportClient = StaticCast<FEditorViewportClient*>(GEditor->GetActiveViewport()->GetClient());
+		if (ViewportClient != nullptr)
+		{
+			FSequencerEdMode* SequencerEdMode = (FSequencerEdMode*)(ViewportClient->GetModeTools()->GetActiveMode(FSequencerEdMode::EM_SequencerMode));
+			SequencerEdMode->DrawMeshTransformTrailFromKey(this);
+		}
+	}
+}
+
diff --git a/Engine/Source/Editor/Sequencer/Private/SequencerMeshTrail.cpp b/Engine/Source/Editor/Sequencer/Private/SequencerMeshTrail.cpp
new file mode 100644
--- /dev/null
+++ b/Engine/Source/Editor/Sequencer/Private/SequencerMeshTrail.cpp
@@ -0,0 +1,157 @@
+// Copyright 1998-2017 Epic Games, Inc. All Rights Reserved.
+
+#include "SequencerMeshTrail.h"
+#include "SequencerKeyActor.h"
+#include "ViewportWorldInteraction.h"
+#include "EditorWorldExtension.h"
+#include "Materials/MaterialInstance.h" 
+#include "Materials/MaterialInstanceDynamic.h"
+#include "Components/BillboardComponent.h"
+#include "MovieScene3DTransformSection.h"
+#include "TimerManager.h"
+
+ASequencerMeshTrail::ASequencerMeshTrail()
+	: Super()
+{
+	TrailTime = 0.0f;
+	MaxTrailTime = 1.0f;
+	UBillboardComponent* PrimitiveRootComponent = CreateDefaultSubobject<UBillboardComponent>(TEXT("RootComponent"));
+	PrimitiveRootComponent->bSelectable = false;
+	PrimitiveRootComponent->SetVisibility(false, false);
+	SetRootComponent(PrimitiveRootComponent);
+
+	if (!IsTemplate() && GEditor)
+	{
+		GEditor->GetTimerManager()->SetTimer(TrailUpdate, FTimerDelegate::CreateUObject(this, &ASequencerMeshTrail::UpdateTrailAppearance, 0.0005f), 0.0005f, true);
+	}
+}
+
+void ASequencerMeshTrail::Cleanup()
+{
+	// Destroy all the key actors this trail created
+	for (FKeyActorData KeyMesh : KeyMeshActors)
+	{
+		KeyMesh.KeyActor->Destroy();
+	}
+	this->Destroy();
+}
+
+void ASequencerMeshTrail::AddKeyMeshActor(float KeyTime, const FTransform KeyTransform, class UMovieScene3DTransformSection* TrackSection)
+{
+	ASequencerKeyActor* KeyMeshActor = nullptr;
+	UStaticMesh* KeyEditorMesh = nullptr;
+	UMaterialInstance* KeyEditorMaterial = nullptr;
+	MaxTrailTime = FMath::Max(KeyTime, MaxTrailTime);
+	FKeyActorData* KeyPtr = KeyMeshActors.FindByPredicate([KeyTime](const FKeyActorData InKey)
+	{
+		return FMath::IsNearlyEqual(KeyTime, InKey.Time);
+	});
+
+	// If we don't currently have an actor for this time, create one
+	if (KeyPtr == nullptr)
+	{
+		UViewportWorldInteraction* WorldInteraction = Cast<UViewportWorldInteraction>( GEditor->GetEditorWorldExtensionsManager()->GetEditorWorldExtensions( GetWorld() )->FindExtension( UViewportWorldInteraction::StaticClass() ) );
+		if( WorldInteraction != nullptr )
+		{
+			KeyMeshActor = WorldInteraction->SpawnTransientSceneActor<ASequencerKeyActor>( TEXT( "KeyMesh" ), false );
+			KeyMeshActor->SetActorTransform(KeyTransform);
+			// Destroy all the frame components this trail created
+			for (FFrameComponentData FrameMesh : FrameMeshComponents)
+			{
+				FrameMesh.FrameComponent->DestroyComponent();
+			}
+			FrameMeshComponents.Reset();
+			KeyMeshActor->SetKeyData(TrackSection, KeyTime);
+			KeyMeshActor->SetOwner(this);
+			FKeyActorData NewKey = FKeyActorData(KeyTime, KeyMeshActor);
+			KeyMeshActors.Add(NewKey);
+		}
+	}
+	else
+	{
+		// Just update the transform
+		KeyPtr->KeyActor->SetActorTransform(KeyTransform);
+	}
+}
+
+void ASequencerMeshTrail::AddFrameMeshComponent(const float FrameTime, const FTransform FrameTransform)
+{
+	UStaticMeshComponent* FrameMeshComponent = nullptr;
+	UStaticMesh* FrameEditorMesh = nullptr;
+	UMaterial* FrameEditorMaterial = nullptr;
+
+	FFrameComponentData* FramePtr = FrameMeshComponents.FindByPredicate([FrameTime](const FFrameComponentData InFrame)
+	{
+		return FMath::IsNearlyEqual(FrameTime, InFrame.Time);
+	});
+
+	// If we don't currently have a component for this time, create one
+	if (FramePtr == nullptr)
+	{
+		FrameMeshComponent = NewObject<UStaticMeshComponent>(this);
+		this->AddOwnedComponent(FrameMeshComponent);
+		FrameMeshComponent->RegisterComponent();
+		FrameEditorMesh = LoadObject<UStaticMesh>(nullptr, TEXT("/Engine/VREditor/TransformGizmo/SM_Sequencer_Node"));
+		check(FrameEditorMesh != nullptr);
+		FrameEditorMaterial = LoadObject<UMaterial>(nullptr, TEXT("/Engine/VREditor/TransformGizmo/Main"));
+		check(FrameEditorMaterial != nullptr);
+		FrameMeshComponent->SetStaticMesh(FrameEditorMesh);
+		FrameMeshComponent->CreateAndSetMaterialInstanceDynamicFromMaterial(0, FrameEditorMaterial);
+		FrameMeshComponent->SetMobility(EComponentMobility::Movable);
+		FrameMeshComponent->SetWorldTransform(FrameTransform);
+		FrameMeshComponent->SetCastShadow(false);
+		FrameMeshComponent->SetCollisionEnabled(ECollisionEnabled::NoCollision);
+		FrameMeshComponent->bSelectable = false;
+		FFrameComponentData NewFrame = FFrameComponentData(FrameTime, FrameMeshComponent);
+		FrameMeshComponents.Add(NewFrame);
+	}
+	else
+	{
+		// Just update the transform
+		FramePtr->FrameComponent->SetWorldTransform(FrameTransform);
+	}
+}
+
+void ASequencerMeshTrail::UpdateTrailAppearance(float UpdateTime)
+{
+	if (FMath::IsNearlyEqual(MaxTrailTime, 0.0f))
+	{
+		MaxTrailTime = 1.0f;
+	}
+	// TODO: No hardcoded glow values
+	TrailTime = FMath::Fmod(TrailTime + UpdateTime, MaxTrailTime);
+	for (FFrameComponentData Frame : FrameMeshComponents)
+	{
+		UMaterialInstanceDynamic* FrameMaterial = Cast<UMaterialInstanceDynamic>(Frame.FrameComponent->GetMaterial(0));
+		if (FrameMaterial != nullptr)
+		{
+			float GlowValue = 0.0f;
+			if (FMath::IsNearlyEqual(TrailTime, Frame.Time, 0.1f))
+			{
+				GlowValue = 12.0f;
+			}
+			else
+			{
+				GlowValue = 3.0f;
+			}
+			FrameMaterial->SetScalarParameterValue("GlowAmount", GlowValue);
+		}
+	}
+	for (FKeyActorData Key : KeyMeshActors)
+	{
+		UMaterialInstanceDynamic* KeyMaterial = Cast<UMaterialInstanceDynamic>(Key.KeyActor->GetMeshComponent()->GetMaterial(0));
+		if (KeyMaterial != nullptr)
+		{
+			float GlowValue = 0.0f;
+			if (FMath::IsNearlyEqual(TrailTime, Key.Time, 0.1f))
+			{
+				GlowValue = 20.0f;
+			}
+			else
+			{
+				GlowValue = 5.0f;
+			}
+			KeyMaterial->SetScalarParameterValue("GlowAmount", GlowValue);
+		}
+	}
+}
diff --git a/Engine/Source/Editor/Sequencer/Public/ISequencer.h b/Engine/Source/Editor/Sequencer/Public/ISequencer.h
--- a/Engine/Source/Editor/Sequencer/Public/ISequencer.h
+++ b/Engine/Source/Editor/Sequencer/Public/ISequencer.h
@@ -365,6 +365,15 @@ public:
 	/** @return Returns a widget containing the sequencer's playback controls */
 	virtual TSharedRef<SWidget> MakeTransportControls(bool bExtended) = 0;
 
+	/** Play or toggle playback at the specified play rate */
+	virtual FReply OnPlay(bool bTogglePlay = true, float InPlayRate = 1.f) = 0;
+
+	/** Pause playback */
+	virtual void Pause() = 0;
+
+	/** Getter for sequencer settings */
+	virtual USequencerSettings* GetSequencerSettings() = 0;
+
 	/** Attempt to find a spawned object in the currently focused movie scene, or the template object for the specified binding ID, if possible */
 	virtual UObject* FindSpawnedObjectOrTemplate(const FGuid& BindingId) = 0;
 
diff --git a/Engine/Source/Editor/UnrealEd/Classes/Editor/EditorEngine.h b/Engine/Source/Editor/UnrealEd/Classes/Editor/EditorEngine.h
--- a/Engine/Source/Editor/UnrealEd/Classes/Editor/EditorEngine.h
+++ b/Engine/Source/Editor/UnrealEd/Classes/Editor/EditorEngine.h
@@ -55,6 +55,7 @@ class USoundBase;
 class USoundNode;
 class UTextureRenderTarget2D;
 struct FAnalyticsEventAttribute;
+class UEditorWorldExtensionManager;
 
 //
 // Things to set in mapSetBrush.
@@ -614,6 +615,12 @@ public:
 	/** The feature level we should use when loading or creating a new world */
 	ERHIFeatureLevel::Type DefaultWorldFeatureLevel;
 
+private:
+
+	/** Manager that holds all extensions paired with a world */
+	UPROPERTY()
+	UEditorWorldExtensionManager* EditorWorldExtensionsManager;
+
 protected:
 
 	/** Count of how many PIE instances are waiting to log in */
@@ -1819,7 +1826,7 @@ public:
 	/**
 	*  Returns the Editors world manager instance.
 	*/
-	TSharedRef<class FEditorWorldManager> GetEditorWorldManager() { return EditorWorldManager.ToSharedRef(); }
+	UEditorWorldExtensionManager* GetEditorWorldExtensionsManager() { return EditorWorldExtensionsManager; }
 
 	// Editor specific
 
@@ -2823,9 +2830,6 @@ private:
 	/** The Timer manager for all timer delegates */
 	TSharedPtr<class FTimerManager> TimerManager;
 
-	/** Manager that holds all the worlds contexts with viewport interaction */
-	TSharedPtr<class FEditorWorldManager> EditorWorldManager;
-
 	/** The output log -> message log redirector for use during PIE */
 	TSharedPtr<class FOutputLogErrorsToMessageLogProxy> OutputLogErrorsToMessageLogProxyPtr;
 
diff --git a/Engine/Source/Editor/UnrealEd/Classes/Editor/Transactor.h b/Engine/Source/Editor/UnrealEd/Classes/Editor/Transactor.h
--- a/Engine/Source/Editor/UnrealEd/Classes/Editor/Transactor.h
+++ b/Engine/Source/Editor/UnrealEd/Classes/Editor/Transactor.h
@@ -124,6 +124,8 @@ protected:
 		TArray<FName>		ReferencedNames;
 		/** The object to track */
 		FPersistentObjectRef	Object;
+		/** Custom change to apply to this object to undo this record.  Executing the undo will return an object that can be used to redo the change. */
+		TUniquePtr<FChange> CustomChange;
 		/** Annotation data for the object stored externally */
 		TSharedPtr<ITransactionObjectAnnotation> ObjectAnnotation;
 		/** Array: If an array object, reference to script array */
@@ -163,8 +165,14 @@ protected:
 		// Constructors.
 		FObjectRecord()
 		{}
-		FObjectRecord( FTransaction* Owner, UObject* InObject, FScriptArray* InArray, int32 InIndex, int32 InCount, int32 InOper, int32 InElementSize, STRUCT_DC InDefaultConstructor, STRUCT_AR InSerializer, STRUCT_DTOR InDestructor );
+		FObjectRecord( FTransaction* Owner, UObject* InObject, TUniquePtr<FChange> InCustomChange, FScriptArray* InArray, int32 InIndex, int32 InCount, int32 InOper, int32 InElementSize, STRUCT_DC InDefaultConstructor, STRUCT_AR InSerializer, STRUCT_DTOR InDestructor );
 
+	private:
+		// Non-copyable
+		FObjectRecord( const FObjectRecord& ) = delete;
+		FObjectRecord& operator=( const FObjectRecord& ) = delete;
+
+	public:
 		// Functions.
 		void SerializeContents( FArchive& Ar, int32 InOper );
 		void Restore( FTransaction* Owner );
@@ -349,8 +357,6 @@ protected:
 	bool					bFlip;
 	/** Used to track direction to iterate over transaction's object records. Typically -1 for Undo, 1 for Redo */
 	int32					Inc;
-	/** Count of the number of UModels modified since the last call to FTransaction::Apply */
-	int32					NumModelsModified;
 
 	/** Objects that will be changed directly by the transaction, empty when not transacting */
 	TMap<UObject*, TSharedPtr<ITransactionObjectAnnotation>> ChangedObjects;
@@ -364,9 +370,17 @@ public:
 		,	Inc(-1)
 	{}
 
+private:
+	// Non-copyable
+	FTransaction( const FTransaction& ) = delete;
+	FTransaction& operator=( const FTransaction& ) = delete;
+
+public:
+	
 	// FTransactionBase interface.
 	virtual void SaveObject( UObject* Object ) override;
 	virtual void SaveArray( UObject* Object, FScriptArray* Array, int32 Index, int32 Count, int32 Oper, int32 ElementSize, STRUCT_DC DefaultConstructor, STRUCT_AR Serializer, STRUCT_DTOR Destructor ) override;
+	virtual void StoreUndo( UObject* Object, TUniquePtr<FChange> UndoChange ) override;
 	virtual void SetPrimaryObject(UObject* InObject) override;
 
 	/**
@@ -404,12 +418,6 @@ public:
 	/** Used by GC to collect referenced objects. */
 	void AddReferencedObjects( FReferenceCollector& Collector );
 
-	/** Returns the number of models that were modified by the last call to FTransaction::Apply(). */
-	int32 GetNumModelsModified() const
-	{
-		return NumModelsModified;
-	}
-
 	/**
 	 * Get all the objects that are part of this transaction.
 	 * @param	Objects		[out] Receives the object list.  Previous contents are cleared.
diff --git a/Engine/Source/Editor/UnrealEd/Classes/Settings/EditorExperimentalSettings.h b/Engine/Source/Editor/UnrealEd/Classes/Settings/EditorExperimentalSettings.h
--- a/Engine/Source/Editor/UnrealEd/Classes/Settings/EditorExperimentalSettings.h
+++ b/Engine/Source/Editor/UnrealEd/Classes/Settings/EditorExperimentalSettings.h
@@ -125,18 +125,10 @@ public:
 	UPROPERTY(EditAnywhere, config, Category = Tools, meta = (DisplayName = "Use OpenCL for convex hull decomposition"))
 	bool bUseOpenCLForConvexHullDecomp;
 
-	/** Enables a preview of the Unreal Editor in VR.  This adds a new tool bar button that allows you to toggle into "VR Mode" instantly.  This feature is still in development, but your feedback is appreciated! */
-	UPROPERTY(EditAnywhere, config, Category=VR, meta = (DisplayName="Enable VR Editing"))
-	bool bEnableVREditing;
-
-	/**If true, wearing a Vive or Oculus Rift headset will automatically enter VR Editing mode if Enable VR Editing is true. */
+	/**If true, wearing a Vive or Oculus Rift headset will automatically enter VR Editing mode */
 	UPROPERTY(EditAnywhere, config, Category = VR, meta = (DisplayName = "Enable VR Mode Auto-Entry"))
 	bool bEnableAutoVREditMode;
 
-	/**If true, the tutorial window should always appear when the VR Editor is launched. */
-	UPROPERTY(EditAnywhere, config, Category = VR, meta = (DisplayName = "Always Show VR Tutorial at Start"))
-	bool bAlwaysShowVRTutorial;
-
 	/** Allows editing of potentially unsafe properties during PIE. Advanced use only - use with caution. */
 	UPROPERTY(EditAnywhere, config, Category = Tools, meta = (DisplayName = "Allow editing of potentially unsafe properties."))
 	bool bAllowPotentiallyUnsafePropertyEditing;
